AssetDefinition - Doing some additional tweaking to the base asset definition type, and shifting some stuff around, fixing some unity build issues...etc. *mop* *mop* *mop*

diff --git a/Engine/Plugins/Editor/EngineAssetDefinitions/Source/Private/AssetDefinition_Blueprint.cpp b/Engine/Plugins/Editor/EngineAssetDefinitions/Source/Private/AssetDefinition_Blueprint.cpp
new file mode 100644
--- /dev/null
+++ b/Engine/Plugins/Editor/EngineAssetDefinitions/Source/Private/AssetDefinition_Blueprint.cpp
@@ -0,0 +1,165 @@
+// Copyright Epic Games, Inc. All Rights Reserved.
+
+#include "AssetDefinition_Blueprint.h"
+
+#include "AssetDefinitionRegistry.h"
+#include "AssetToolsModule.h"
+#include "BlueprintEditorModule.h"
+#include "ContentBrowserMenuContexts.h"
+#include "ContentBrowserModule.h"
+#include "ToolMenus.h"
+#include "Factories/BlueprintFactory.h"
+#include "Kismet2/KismetEditorUtilities.h"
+#include "IContentBrowserSingleton.h"
+#include "Algo/AllOf.h"
+#include "Logging/MessageLog.h"
+#include "Misc/MessageDialog.h"
+
+#define LOCTEXT_NAMESPACE "UAssetDefinition_Blueprint"
+
+namespace MenuExtension_Blueprint
+{
+	static bool CanExecuteNewDerivedBlueprint(const FToolMenuContext& MenuContext, const FAssetData* SelectedBlueprintPtr)
+	{
+		const uint32 BPFlags = SelectedBlueprintPtr->GetTagValueRef<uint32>(FBlueprintTags::ClassFlags);
+		if ((BPFlags & (CLASS_Deprecated)) == 0)
+		{
+			return true;
+		}
+
+		return false;
+	}
+
+	static void ExecuteNewDerivedBlueprint(const FToolMenuContext& MenuContext, const FAssetData* SelectedBlueprintPtr)
+	{
+		if (UBlueprint* ParentBlueprint = Cast<UBlueprint>(SelectedBlueprintPtr->GetAsset()))
+		{
+			UClass* TargetParentClass = ParentBlueprint->GeneratedClass;
+
+			if (!FKismetEditorUtilities::CanCreateBlueprintOfClass(TargetParentClass))
+			{
+				FMessageDialog::Open( EAppMsgType::Ok, LOCTEXT("InvalidClassToMakeBlueprintFrom", "Invalid class with which to make a Blueprint."));
+				return;
+			}
+
+			FString Name;
+			FString PackageName;
+			FAssetToolsModule& AssetToolsModule = FModuleManager::LoadModuleChecked<FAssetToolsModule>("AssetTools");
+			AssetToolsModule.Get().CreateUniqueAssetName(ParentBlueprint->GetOutermost()->GetName(), TEXT("_Child"), PackageName, Name);
+			const FString PackagePath = FPackageName::GetLongPackagePath(PackageName);
+
+			const UAssetDefinition_Blueprint* BlueprintAssetDefinition = Cast<UAssetDefinition_Blueprint>(UAssetDefinitionRegistry::Get()->GetAssetDefinitionForClass(ParentBlueprint->GetClass()));
+			if (BlueprintAssetDefinition)
+			{
+				UFactory* Factory = BlueprintAssetDefinition->GetFactoryForBlueprintType(ParentBlueprint);
+
+				FContentBrowserModule& ContentBrowserModule = FModuleManager::LoadModuleChecked<FContentBrowserModule>("ContentBrowser");
+				ContentBrowserModule.Get().CreateNewAsset(Name, PackagePath, ParentBlueprint->GetClass(), Factory);
+			}
+		}
+	}
+
+	static void ExecuteEditDefaults(const FToolMenuContext& MenuContext, TArray<FAssetData> BlueprintAssets)
+	{
+		TArray<UBlueprint*> Blueprints;
+
+		FMessageLog EditorErrors("EditorErrors");
+		EditorErrors.NewPage(LOCTEXT("ExecuteEditDefaultsNewLogPage", "Loading Blueprints"));
+
+		for (const FAssetData& BlueprintAsset : BlueprintAssets)
+		{
+			if (UBlueprint* Object = Cast<UBlueprint>(BlueprintAsset.GetAsset()))
+			{
+				// If the blueprint is valid, allow it to be added to the list, otherwise log the error.
+				if ( Object->SkeletonGeneratedClass && Object->GeneratedClass )
+				{
+					Blueprints.Add(Object);
+				}
+				else
+				{
+					FFormatNamedArguments Arguments;
+					Arguments.Add(TEXT("ObjectName"), FText::FromString(Object->GetName()));
+					EditorErrors.Error(FText::Format(LOCTEXT("LoadBlueprint_FailedLog", "{ObjectName} could not be loaded because it derives from an invalid class.  Check to make sure the parent class for this blueprint hasn't been removed!"), Arguments ) );
+				}
+			}
+		}
+
+		if ( Blueprints.Num() > 0 )
+		{
+			FBlueprintEditorModule& BlueprintEditorModule = FModuleManager::LoadModuleChecked<FBlueprintEditorModule>( "Kismet" );
+			TSharedRef< IBlueprintEditor > NewBlueprintEditor = BlueprintEditorModule.CreateBlueprintEditor(  EToolkitMode::Standalone, TSharedPtr<IToolkitHost>(), Blueprints );
+		}
+
+		// Report errors
+		EditorErrors.Notify(LOCTEXT("OpenDefaults_Failed", "Opening Class Defaults Failed!"));
+	}
+	
+	static FDelayedAutoRegisterHelper DelayedAutoRegister(EDelayedRegisterRunPhase::EndOfEngineInit, []{ 
+		UToolMenus::RegisterStartupCallback(FSimpleMulticastDelegate::FDelegate::CreateLambda([]()
+		{
+			FToolMenuOwnerScoped OwnerScoped(UE_MODULE_NAME);
+			UToolMenu* Menu = UE::ContentBrowser::ExtendToolMenu_AssetContextMenu(UBlueprint::StaticClass());
+	        
+			FToolMenuSection& Section = Menu->FindOrAddSection("GetAssetActions");
+				Section.AddDynamicEntry("GetAssetActions_Blueprint", FNewToolMenuSectionDelegate::CreateLambda([](FToolMenuSection& InSection)
+				{
+					if (const UContentBrowserAssetContextMenuContext* Context = UContentBrowserAssetContextMenuContext::FindContextWithAssets(InSection))
+					{
+						// TODO NDarnell re: EIncludeSubclasses::No, Temporary - Need to ensure we don't have duplicates for now, because no all subclasses of blueprint are of this class yet.
+						if (const FAssetData* SelectedBlueprintPtr = Context->GetSingleSelectedAssetOfType(UBlueprint::StaticClass(), EIncludeSubclasses::No))
+						{
+							const TAttribute<FText> Label = LOCTEXT("Blueprint_NewDerivedBlueprint", "Create Child Blueprint Class");
+							const TAttribute<FText> ToolTip = TAttribute<FText>::CreateLambda([SelectedBlueprintPtr]()
+							{
+								const uint32 BPFlags = SelectedBlueprintPtr->GetTagValueRef<uint32>(FBlueprintTags::ClassFlags);
+								if ((BPFlags & (CLASS_Deprecated)) == 0)
+								{
+									return LOCTEXT("Blueprint_NewDerivedBlueprintTooltip", "Creates a Child Blueprint Class based on the current Blueprint, allowing you to create variants easily.");
+								}
+								else
+								{
+									return LOCTEXT("Blueprint_NewDerivedBlueprintIsDeprecatedTooltip", "Blueprint class is deprecated, cannot derive a child Blueprint!");
+								}
+							});
+							const FSlateIcon Icon = FSlateIcon(FAppStyle::GetAppStyleSetName(), "ClassIcon.Blueprint");
+
+							FToolUIAction DeriveNewBlueprint;
+							DeriveNewBlueprint.ExecuteAction = FToolMenuExecuteAction::CreateStatic(&ExecuteNewDerivedBlueprint, SelectedBlueprintPtr);
+							DeriveNewBlueprint.CanExecuteAction = FToolMenuCanExecuteAction::CreateStatic(&CanExecuteNewDerivedBlueprint, SelectedBlueprintPtr);
+							InSection.AddMenuEntry("CreateChildBlueprintClass", Label, ToolTip, Icon, DeriveNewBlueprint);
+						}
+
+						TArray<FAssetData> SelectedBlueprints = Context->GetSelectedAssetsOfType(UBlueprint::StaticClass(), EIncludeSubclasses::No);
+						if (SelectedBlueprints.Num() > 1)
+						{
+							TArray<UClass*> SelectedBlueprintParentClasses;
+							Algo::Transform(SelectedBlueprints, SelectedBlueprintParentClasses, [](const FAssetData& BlueprintAsset){ return UBlueprint::GetBlueprintParentClassFromAssetTags(BlueprintAsset); });
+							
+							// Ensure that all the selected blueprints are actors
+							const bool bAreAllSelectedBlueprintsActors =
+								Algo::AllOf(SelectedBlueprintParentClasses, [](UClass* ParentClass){ return ParentClass->IsChildOf(AActor::StaticClass()); });
+							
+							if (bAreAllSelectedBlueprintsActors)
+							{
+								const TAttribute<FText> Label = LOCTEXT("Blueprint_EditDefaults", "Edit Shared Defaults");
+								const TAttribute<FText> ToolTip = LOCTEXT("Blueprint_EditDefaultsTooltip", "Edit the shared default properties of the selected actor blueprints.");
+								const FSlateIcon Icon = FSlateIcon(FAppStyle::GetAppStyleSetName(), "Kismet.Tabs.BlueprintDefaults");
+								const FToolMenuExecuteAction UIAction = FToolMenuExecuteAction::CreateStatic(&ExecuteEditDefaults, MoveTemp(SelectedBlueprints));
+								InSection.AddMenuEntry("Blueprint_EditDefaults", Label, ToolTip, Icon, UIAction);
+							}
+						}
+					}
+				}));
+		}));
+	});
+}
+
+
+UFactory* UAssetDefinition_Blueprint::GetFactoryForBlueprintType(UBlueprint* InBlueprint) const
+{
+	UBlueprintFactory* BlueprintFactory = NewObject<UBlueprintFactory>();
+	BlueprintFactory->ParentClass = InBlueprint->GeneratedClass;
+	return BlueprintFactory;
+}
+
+#undef LOCTEXT_NAMESPACE
diff --git a/Engine/Plugins/Editor/EngineAssetDefinitions/Source/Private/AssetDefinition_DataAsset.cpp b/Engine/Plugins/Editor/EngineAssetDefinitions/Source/Private/AssetDefinition_DataAsset.cpp
--- a/Engine/Plugins/Editor/EngineAssetDefinitions/Source/Private/AssetDefinition_DataAsset.cpp
+++ b/Engine/Plugins/Editor/EngineAssetDefinitions/Source/Private/AssetDefinition_DataAsset.cpp
@@ -10,6 +10,7 @@
 #include "SDetailsDiff.h"
 #include "ToolMenus.h"
 #include "Kismet2/SClassPickerDialog.h"
+#include "Engine/Engine.h"
 
 #define LOCTEXT_NAMESPACE "UAssetDefinition_DataAsset"
 
@@ -107,7 +108,6 @@ namespace MenuExtension_DataAsset
 			}));
 		}));
 	});
-
 }
 
 
@@ -144,78 +144,4 @@ EAssetCommandResult UAssetDefinition_DataAsset::PerformAssetDiff(const FAssetDif
 	return EAssetCommandResult::Handled;
 }
 
-void UAssetDefinition_DataAsset::OnRegistered()
-{
-	UToolMenus::RegisterStartupCallback(FSimpleMulticastDelegate::FDelegate::CreateUObject(this, &ThisClass::RegisterMenus));
-}
-
-void UAssetDefinition_DataAsset::RegisterMenus()
-{
-	return;
-	//FToolMenuOwnerScoped MenuOwner(GetClass()->GetOutermost()->GetLoadedPath());
-
-	UToolMenu* Menu = UE::ContentBrowser::ExtendToolMenu_AssetContextMenu(GetAssetClass());
-	
-	FToolMenuSection& Section = Menu->FindOrAddSection("GetAssetActions");
-	Section.AddDynamicEntry("GetAssetActions_DataAsset", FNewToolMenuSectionDelegate::CreateWeakLambda(this, [this](FToolMenuSection& InSection)
-	{
-		{
-			const TAttribute<FText> Label = LOCTEXT("DataAsset_ChangeClass", "Convert to Different DataAsset Type");
-			const TAttribute<FText> ToolTip = LOCTEXT("DataAsset_ChangeClassTip", "Change the class these Data Assets are subclassed from.");
-			const FSlateIcon Icon = FSlateIcon(FAppStyle::GetAppStyleSetName(), "ClassIcon.DataAsset");
-			const FToolMenuExecuteAction UIAction = FToolMenuExecuteAction::CreateUObject(this, &ThisClass::ExecuteChangeDataAssetClass);
-
-			InSection.AddMenuEntry("DataAsset_ChangeClass", Label, ToolTip, Icon, UIAction);
-		}
-	}));
-}
-
-void UAssetDefinition_DataAsset::ExecuteChangeDataAssetClass(const FToolMenuContext& MenuContext)
-{
-	/*
-	if (const UContentBrowserAssetContextMenuContext* Context = UContentBrowserAssetContextMenuContext::FindContextWithAssets(MenuContext))
-	{
-		TArray<UDataAsset*> DataAssets = Context->LoadSelectedObjects<UDataAsset>();
-
-		const FText TitleText = LOCTEXT("DataAsset_PickNewDataAssetClass", "Pick New DataAsset Class");
-		FClassViewerInitializationOptions Options;
-		Options.ClassFilters.Add(MakeShared<FNewNodeClassFilter>(UDataAsset::StaticClass()));
-		UClass* OutNewDataAssetClass = nullptr;
-		const bool bPressedOk = SClassPickerDialog::PickClass(TitleText, Options, OutNewDataAssetClass, UDataAsset::StaticClass());
-
-		if (bPressedOk && OutNewDataAssetClass != nullptr)
-		{
-			for (TWeakObjectPtr<UDataAsset> DataAssetPtr : DataAssets)
-			{
-				if (UDataAsset* OldDataAsset = DataAssetPtr.Get())
-				{
-					if (OldDataAsset && OldDataAsset->IsValidLowLevel())
-					{
-						FName ObjectName = OldDataAsset->GetFName();
-						UObject* Outer = OldDataAsset->GetOuter();
-						OldDataAsset->Rename(nullptr, GetTransientPackage(), REN_DoNotDirty | REN_DontCreateRedirectors);
-
-						UObject* NewDataAsset = NewObject<UObject>(Outer, OutNewDataAssetClass, ObjectName, OldDataAsset->GetFlags());
-
-						// Migrate Data
-						{
-							UEngine::FCopyPropertiesForUnrelatedObjectsParams CopyOptions;
-							CopyOptions.bNotifyObjectReplacement = true;
-							UEngine::CopyPropertiesForUnrelatedObjects(OldDataAsset, NewDataAsset, CopyOptions);
-						}
-
-						NewDataAsset->MarkPackageDirty();
-
-						// Consolidate or "Replace" the old object with the new object for any living references.
-						bool bShowDeleteConfirmation = false;
-						TArray<UObject*> OldDataAssetArray({ (UObject*)OldDataAsset });
-						ObjectTools::ConsolidateObjects(NewDataAsset, OldDataAssetArray, bShowDeleteConfirmation);
-					}
-				}
-			}
-		}
-	}
-	*/
-}
-
 #undef LOCTEXT_NAMESPACE
diff --git a/Engine/Plugins/Editor/EngineAssetDefinitions/Source/Private/AssetDefinition_DataAsset.h b/Engine/Plugins/Editor/EngineAssetDefinitions/Source/Private/AssetDefinition_DataAsset.h
deleted file mode 100644
diff --git a/Engine/Plugins/Editor/EngineAssetDefinitions/Source/Private/AssetDefinition_Redirector.cpp b/Engine/Plugins/Editor/EngineAssetDefinitions/Source/Private/AssetDefinition_Redirector.cpp
--- a/Engine/Plugins/Editor/EngineAssetDefinitions/Source/Private/AssetDefinition_Redirector.cpp
+++ b/Engine/Plugins/Editor/EngineAssetDefinitions/Source/Private/AssetDefinition_Redirector.cpp
@@ -8,109 +8,111 @@
 
 #define LOCTEXT_NAMESPACE "UAssetDefinition_Redirector"
 
-void UAssetDefinition_Redirector::OnRegistered()
+namespace RedirectorUtilities
 {
-	UToolMenus::RegisterStartupCallback(FSimpleMulticastDelegate::FDelegate::CreateUObject(this, &ThisClass::RegisterMenus));
-}
-
-void UAssetDefinition_Redirector::RegisterMenus()
-{
-	//FToolMenuOwnerScoped MenuOwner(GetClass()->GetOutermost()->GetLoadedPath());
-
-	UToolMenu* Menu = UE::ContentBrowser::ExtendToolMenu_AssetContextMenu(GetAssetClass());
-	
-	FToolMenuSection& Section = Menu->FindOrAddSection("GetAssetActions");
-	Section.AddDynamicEntry("GetAssetActions_ObjectRedirector", FNewToolMenuSectionDelegate::CreateWeakLambda(this, [this](FToolMenuSection& InSection)
+	static void FindTargets(const TArray<UObjectRedirector*>& Redirectors)
 	{
+		TArray<UObject*> ObjectsToSync;
+		for (auto RedirectorIt = Redirectors.CreateConstIterator(); RedirectorIt; ++RedirectorIt)
 		{
-			const TAttribute<FText> Label = LOCTEXT("Redirector_FindTarget", "Find Target");
-			const TAttribute<FText> ToolTip = LOCTEXT("Redirector_FindTargetTooltip", "Finds the asset that this redirector targets in the asset tree.");
-			const FSlateIcon Icon = FSlateIcon(FAppStyle::GetAppStyleSetName(), "ClassIcon.ObjectRedirector");
-			const FToolMenuExecuteAction UIAction = FToolMenuExecuteAction::CreateUObject(this, &ThisClass::ExecuteFindTarget);
-
-			InSection.AddMenuEntry("Redirector_FindTarget", Label, ToolTip, Icon, UIAction);
+			const UObjectRedirector* Redirector = *RedirectorIt;
+			if ( Redirector && Redirector->DestinationObject )
+			{
+				ObjectsToSync.Add(Redirector->DestinationObject);
+			}
 		}
-		{
-			const TAttribute<FText> Label = LOCTEXT("Redirector_FixUp", "Fix Up");
-			const TAttribute<FText> ToolTip = LOCTEXT("Redirector_FixUpTooltip", "Finds referencers to selected redirectors and resaves them if possible, then deletes any redirectors that had all their referencers fixed.");
-			const FSlateIcon Icon = FSlateIcon(FAppStyle::GetAppStyleSetName(), "ClassIcon.ObjectRedirector");
-			const FToolMenuExecuteAction UIAction = FToolMenuExecuteAction::CreateUObject(this, &ThisClass::ExecuteFixUp, true);
 
-			InSection.AddMenuEntry("Redirector_FixUp", Label, ToolTip, Icon, UIAction);
-		}
+		if ( ObjectsToSync.Num() > 0 )
 		{
-			const TAttribute<FText> Label = LOCTEXT("Redirector_FixUp_KeepingRedirector", "Fix Up (Keep Redirector)");
-			const TAttribute<FText> ToolTip = LOCTEXT("Redirector_FixUp_KeepingRedirectorTooltip", "Finds referencers to selected redirectors and resaves them if possible.");
-			const FSlateIcon Icon = FSlateIcon(FAppStyle::GetAppStyleSetName(), "ClassIcon.MyAsset");
-			const FToolMenuExecuteAction UIAction = FToolMenuExecuteAction::CreateUObject(this, &ThisClass::ExecuteFixUp, false);
-
-			InSection.AddMenuEntry("Redirector_FixUp_KeepingRedirector", Label, ToolTip, Icon, UIAction);
+			IAssetTools& AssetTools = FModuleManager::LoadModuleChecked<FAssetToolsModule>("AssetTools").Get();
+			AssetTools.SyncBrowserToAssets(ObjectsToSync);
 		}
-	}));
+	}
 }
 
-EAssetCommandResult UAssetDefinition_Redirector::ActivateAssets(const FAssetActivateArgs& ActivateArgs) const
+namespace MenuExtension_Redirector
 {
-	if ( ActivateArgs.ActivationMethod == EAssetActivationMethod::DoubleClicked || ActivateArgs.ActivationMethod == EAssetActivationMethod::Opened )
+	static void ExecuteFindTarget(const FToolMenuContext& MenuContext)
 	{
-		// Sync to the target instead of opening an editor when double clicked
-		TArray<UObjectRedirector*> Redirectors = ActivateArgs.LoadObjects<UObjectRedirector>();
-		if ( Redirectors.Num() > 0 )
+		if (const UContentBrowserAssetContextMenuContext* Context = UContentBrowserAssetContextMenuContext::FindContextWithAssets(MenuContext))
 		{
-			FindTargets(Redirectors);
-			return EAssetCommandResult::Handled;
+			TArray<UObjectRedirector*> Redirectors = Context->LoadSelectedObjects<UObjectRedirector>();
+			RedirectorUtilities::FindTargets(Redirectors);
 		}
 	}
 
-	return EAssetCommandResult::Unhandled;
-}
-
-void UAssetDefinition_Redirector::ExecuteFindTarget(const FToolMenuContext& MenuContext)
-{
-	if (const UContentBrowserAssetContextMenuContext* Context = UContentBrowserAssetContextMenuContext::FindContextWithAssets(MenuContext))
-	{
-		TArray<UObjectRedirector*> Redirectors = Context->LoadSelectedObjects<UObjectRedirector>();
-		FindTargets(Redirectors);
-	}
-}
-
-void UAssetDefinition_Redirector::ExecuteFixUp(const FToolMenuContext& MenuContext, bool bDeleteAssets)
-{
-	// This will fix references to selected redirectors, except in the following cases:
-	// Redirectors referenced by unloaded maps will not be fixed up, but any references to it that can be fixed up will
-	// Redirectors referenced by code will not be completely fixed up
-	// Redirectors that are not at head revision or checked out by another user will not be completely fixed up
-	// Redirectors whose referencers are not at head revision, are checked out by another user, or are refused to be checked out will not be completely fixed up.
-
-	if (const UContentBrowserAssetContextMenuContext* Context = UContentBrowserAssetContextMenuContext::FindContextWithAssets(MenuContext))
+	static void ExecuteFixUp(const FToolMenuContext& MenuContext, bool bDeleteAssets)
 	{
-		TArray<UObjectRedirector*> Redirectors = Context->LoadSelectedObjects<UObjectRedirector>();
+		// This will fix references to selected redirectors, except in the following cases:
+		// Redirectors referenced by unloaded maps will not be fixed up, but any references to it that can be fixed up will
+		// Redirectors referenced by code will not be completely fixed up
+		// Redirectors that are not at head revision or checked out by another user will not be completely fixed up
+		// Redirectors whose referencers are not at head revision, are checked out by another user, or are refused to be checked out will not be completely fixed up.
 
-		if (Redirectors.Num() > 0)
+		if (const UContentBrowserAssetContextMenuContext* Context = UContentBrowserAssetContextMenuContext::FindContextWithAssets(MenuContext))
 		{
-			IAssetTools& AssetTools = FModuleManager::LoadModuleChecked<FAssetToolsModule>("AssetTools").Get();
-			AssetTools.FixupReferencers(Redirectors, /*bCheckoutDialogPrompt=*/ true, bDeleteAssets ? ERedirectFixupMode::DeleteFixedUpRedirectors : ERedirectFixupMode::LeaveFixedUpRedirectors);
+			TArray<UObjectRedirector*> Redirectors = Context->LoadSelectedObjects<UObjectRedirector>();
+
+			if (Redirectors.Num() > 0)
+			{
+				IAssetTools& AssetTools = FModuleManager::LoadModuleChecked<FAssetToolsModule>("AssetTools").Get();
+				AssetTools.FixupReferencers(Redirectors, /*bCheckoutDialogPrompt=*/ true, bDeleteAssets ? ERedirectFixupMode::DeleteFixedUpRedirectors : ERedirectFixupMode::LeaveFixedUpRedirectors);
+			}
 		}
 	}
+	
+	static FDelayedAutoRegisterHelper DelayedAutoRegister(EDelayedRegisterRunPhase::EndOfEngineInit, []{ 
+		UToolMenus::RegisterStartupCallback(FSimpleMulticastDelegate::FDelegate::CreateLambda([]()
+		{
+			FToolMenuOwnerScoped OwnerScoped(UE_MODULE_NAME);
+			UToolMenu* Menu = UE::ContentBrowser::ExtendToolMenu_AssetContextMenu(UObjectRedirector::StaticClass());
+	        
+			FToolMenuSection& Section = Menu->FindOrAddSection("GetAssetActions");
+			Section.AddDynamicEntry("GetAssetActions_ObjectRedirector", FNewToolMenuSectionDelegate::CreateLambda([](FToolMenuSection& InSection)
+			{
+				{
+					const TAttribute<FText> Label = LOCTEXT("Redirector_FindTarget", "Find Target");
+					const TAttribute<FText> ToolTip = LOCTEXT("Redirector_FindTargetTooltip", "Finds the asset that this redirector targets in the asset tree.");
+					const FSlateIcon Icon = FSlateIcon(FAppStyle::GetAppStyleSetName(), "ClassIcon.ObjectRedirector");
+					const FToolMenuExecuteAction UIAction = FToolMenuExecuteAction::CreateStatic(&ExecuteFindTarget);
+
+					InSection.AddMenuEntry("Redirector_FindTarget", Label, ToolTip, Icon, UIAction);
+				}
+				{
+					const TAttribute<FText> Label = LOCTEXT("Redirector_FixUp", "Fix Up");
+					const TAttribute<FText> ToolTip = LOCTEXT("Redirector_FixUpTooltip", "Finds referencers to selected redirectors and resaves them if possible, then deletes any redirectors that had all their referencers fixed.");
+					const FSlateIcon Icon = FSlateIcon(FAppStyle::GetAppStyleSetName(), "ClassIcon.ObjectRedirector");
+					const FToolMenuExecuteAction UIAction = FToolMenuExecuteAction::CreateStatic(&ExecuteFixUp, true);
+
+					InSection.AddMenuEntry("Redirector_FixUp", Label, ToolTip, Icon, UIAction);
+				}
+				{
+					const TAttribute<FText> Label = LOCTEXT("Redirector_FixUp_KeepingRedirector", "Fix Up (Keep Redirector)");
+					const TAttribute<FText> ToolTip = LOCTEXT("Redirector_FixUp_KeepingRedirectorTooltip", "Finds referencers to selected redirectors and resaves them if possible.");
+					const FSlateIcon Icon = FSlateIcon(FAppStyle::GetAppStyleSetName(), "ClassIcon.MyAsset");
+					const FToolMenuExecuteAction UIAction = FToolMenuExecuteAction::CreateStatic(&ExecuteFixUp, false);
+
+					InSection.AddMenuEntry("Redirector_FixUp_KeepingRedirector", Label, ToolTip, Icon, UIAction);
+				}
+			}));
+		}));
+	});
 }
 
-void UAssetDefinition_Redirector::FindTargets(const TArray<UObjectRedirector*>& Redirectors) const
+EAssetCommandResult UAssetDefinition_Redirector::ActivateAssets(const FAssetActivateArgs& ActivateArgs) const
 {
-	TArray<UObject*> ObjectsToSync;
-	for (auto RedirectorIt = Redirectors.CreateConstIterator(); RedirectorIt; ++RedirectorIt)
+	if ( ActivateArgs.ActivationMethod == EAssetActivationMethod::DoubleClicked || ActivateArgs.ActivationMethod == EAssetActivationMethod::Opened )
 	{
-		const UObjectRedirector* Redirector = *RedirectorIt;
-		if ( Redirector && Redirector->DestinationObject )
+		// Sync to the target instead of opening an editor when double clicked
+		TArray<UObjectRedirector*> Redirectors = ActivateArgs.LoadObjects<UObjectRedirector>();
+		if ( Redirectors.Num() > 0 )
 		{
-			ObjectsToSync.Add(Redirector->DestinationObject);
+			RedirectorUtilities::FindTargets(Redirectors);
+			return EAssetCommandResult::Handled;
 		}
 	}
 
-	if ( ObjectsToSync.Num() > 0 )
-	{
-		IAssetTools& AssetTools = FModuleManager::LoadModuleChecked<FAssetToolsModule>("AssetTools").Get();
-		AssetTools.SyncBrowserToAssets(ObjectsToSync);
-	}
+	return EAssetCommandResult::Unhandled;
 }
 
 #undef LOCTEXT_NAMESPACE
diff --git a/Engine/Plugins/Editor/EngineAssetDefinitions/Source/Private/AssetDefinition_Redirector.h b/Engine/Plugins/Editor/EngineAssetDefinitions/Source/Private/AssetDefinition_Redirector.h
deleted file mode 100644
diff --git a/Engine/Plugins/Editor/EngineAssetDefinitions/Source/Private/AssetDefinition_World.cpp b/Engine/Plugins/Editor/EngineAssetDefinitions/Source/Private/AssetDefinition_World.cpp
--- a/Engine/Plugins/Editor/EngineAssetDefinitions/Source/Private/AssetDefinition_World.cpp
+++ b/Engine/Plugins/Editor/EngineAssetDefinitions/Source/Private/AssetDefinition_World.cpp
@@ -4,6 +4,9 @@
 
 #include "FileHelpers.h"
 #include "ThumbnailRendering/WorldThumbnailInfo.h"
+#include "Misc/MessageDialog.h"
+#include "Engine/Level.h"
+#include "Editor.h"
 
 #define LOCTEXT_NAMESPACE "UAssetDefinition_World"
 
diff --git a/Engine/Plugins/Editor/EngineAssetDefinitions/Source/Private/AssetDefinition_World.h b/Engine/Plugins/Editor/EngineAssetDefinitions/Source/Private/AssetDefinition_World.h
deleted file mode 100644
diff --git a/Engine/Plugins/Editor/EngineAssetDefinitions/Source/Public/AssetDefinition_Blueprint.h b/Engine/Plugins/Editor/EngineAssetDefinitions/Source/Public/AssetDefinition_Blueprint.h
new file mode 100644
--- /dev/null
+++ b/Engine/Plugins/Editor/EngineAssetDefinitions/Source/Public/AssetDefinition_Blueprint.h
@@ -0,0 +1,28 @@
+// Copyright Epic Games, Inc. All Rights Reserved.
+
+#pragma once
+
+#include "CoreMinimal.h"
+#include "AssetDefinitionDefault.h"
+#include "Engine/Blueprint.h"
+
+#include "AssetDefinition_Blueprint.generated.h"
+
+struct FToolMenuContext;
+struct FAssetData;
+class UFactory;
+
+UCLASS()
+class ENGINEASSETDEFINITIONS_API UAssetDefinition_Blueprint : public UAssetDefinitionDefault
+{
+	GENERATED_BODY()
+
+public:
+	// UAssetDefinition Implementation
+	virtual FText GetAssetDisplayName() const override { return NSLOCTEXT("AssetTypeActions", "AssetTypeActions_Blueprint", "Blueprint Class"); }
+	virtual FLinearColor GetAssetColor() const override { return FLinearColor(FColor( 63, 126, 255 )); }
+	virtual TSoftClassPtr<UObject> GetAssetClass() const override { return UBlueprint::StaticClass(); }
+
+public:
+	virtual UFactory* GetFactoryForBlueprintType(UBlueprint* InBlueprint) const;
+};
diff --git a/Engine/Plugins/Editor/EngineAssetDefinitions/Source/Public/AssetDefinition_DataAsset.h b/Engine/Plugins/Editor/EngineAssetDefinitions/Source/Public/AssetDefinition_DataAsset.h
new file mode 100644
--- /dev/null
+++ b/Engine/Plugins/Editor/EngineAssetDefinitions/Source/Public/AssetDefinition_DataAsset.h
@@ -0,0 +1,24 @@
+// Copyright Epic Games, Inc. All Rights Reserved.
+
+#pragma once
+
+#include "CoreMinimal.h"
+#include "AssetDefinitionDefault.h"
+#include "Engine/DataAsset.h"
+
+#include "AssetDefinition_DataAsset.generated.h"
+
+UCLASS()
+class ENGINEASSETDEFINITIONS_API UAssetDefinition_DataAsset : public UAssetDefinitionDefault
+{
+	GENERATED_BODY()
+
+public:
+	// UAssetDefinition Implementation
+	virtual FText GetAssetDisplayName() const override { return NSLOCTEXT("AssetDefinition", "DataAsset", "Data Asset"); }
+	virtual FText GetAssetDisplayName(const FAssetData& AssetData) const override;
+	virtual FLinearColor GetAssetColor() const override { return FLinearColor(FColor(201, 29, 85)); }
+	virtual TSoftClassPtr<UObject> GetAssetClass() const override { return UDataAsset::StaticClass(); }
+	
+	virtual EAssetCommandResult PerformAssetDiff(const FAssetDiffArgs& DiffArgs) const override;
+};
diff --git a/Engine/Plugins/Editor/EngineAssetDefinitions/Source/Public/AssetDefinition_Redirector.h b/Engine/Plugins/Editor/EngineAssetDefinitions/Source/Public/AssetDefinition_Redirector.h
new file mode 100644
--- /dev/null
+++ b/Engine/Plugins/Editor/EngineAssetDefinitions/Source/Public/AssetDefinition_Redirector.h
@@ -0,0 +1,22 @@
+// Copyright Epic Games, Inc. All Rights Reserved.
+
+#pragma once
+
+#include "CoreMinimal.h"
+#include "AssetDefinitionDefault.h"
+#include "UObject/ObjectRedirector.h"
+
+#include "AssetDefinition_Redirector.generated.h"
+
+UCLASS()
+class ENGINEASSETDEFINITIONS_API UAssetDefinition_Redirector : public UAssetDefinitionDefault
+{
+	GENERATED_BODY()
+
+public:
+	// UAssetDefinition Implementation
+	virtual FText GetAssetDisplayName() const override { return NSLOCTEXT("AssetDefinition", "Redirector", "Redirector"); }
+	virtual FLinearColor GetAssetColor() const override { return FLinearColor(FColor(128, 128, 128)); }
+	virtual TSoftClassPtr<UObject> GetAssetClass() const override { return UObjectRedirector::StaticClass(); }
+	virtual EAssetCommandResult ActivateAssets(const FAssetActivateArgs& ActivateArgs) const override;
+};
diff --git a/Engine/Plugins/Editor/EngineAssetDefinitions/Source/Public/AssetDefinition_World.h b/Engine/Plugins/Editor/EngineAssetDefinitions/Source/Public/AssetDefinition_World.h
new file mode 100644
--- /dev/null
+++ b/Engine/Plugins/Editor/EngineAssetDefinitions/Source/Public/AssetDefinition_World.h
@@ -0,0 +1,34 @@
+// Copyright Epic Games, Inc. All Rights Reserved.
+
+#pragma once
+
+#include "CoreMinimal.h"
+#include "Engine/World.h"
+#include "AssetDefinitionDefault.h"
+
+#include "AssetDefinition_World.generated.h"
+
+UCLASS()
+class ENGINEASSETDEFINITIONS_API UAssetDefinition_World : public UAssetDefinitionDefault
+{
+	GENERATED_BODY()
+
+public:
+	// UAssetDefinition Begin
+	virtual FText GetAssetDisplayName() const override { return NSLOCTEXT("AssetTypeActions", "AssetTypeActions_World", "Level"); }
+	virtual FLinearColor GetAssetColor() const override { return FAppStyle::Get().GetColor("LevelEditor.AssetColor"); }
+	virtual TSoftClassPtr<UObject> GetAssetClass() const override { return UWorld::StaticClass(); }
+	virtual TConstArrayView<FAssetCategoryPath> GetAssetCategories() const override { return { EAssetCategoryPaths::Basic }; }
+
+	virtual TArray<FAssetData> PrepareToActivateAssets(const FAssetActivateArgs& ActivateArgs) const override;
+	virtual EAssetCommandResult OpenAssets(const FAssetOpenArgs& OpenArgs) const override;
+	
+	virtual FAssetSupportResponse CanRename(const FAssetData& InAsset) const override;
+	virtual FAssetSupportResponse CanDuplicate(const FAssetData& InAsset) const override;
+
+	virtual UThumbnailInfo* LoadThumbnailInfo(const FAssetData& InAsset) const override;
+	// UAssetDefinition End
+	
+public:
+	bool IsPartitionWorldInUse(const FAssetData& InAsset) const;
+};
diff --git a/Engine/Source/Editor/AssetDefinition/Private/AssetDefinitionRegistry.cpp b/Engine/Source/Editor/AssetDefinition/Private/AssetDefinitionRegistry.cpp
--- a/Engine/Source/Editor/AssetDefinition/Private/AssetDefinitionRegistry.cpp
+++ b/Engine/Source/Editor/AssetDefinition/Private/AssetDefinitionRegistry.cpp
@@ -79,7 +79,6 @@ void UAssetDefinitionRegistry::RegisterAssetDefinition(UAssetDefinition* AssetDe
 		if (CanAddDefinition)
 		{
 			AssetDefinitions.Add(SupportedClass, AssetDefinition);
-			AssetDefinition->OnRegistered();
 		}
 		else
 		{
@@ -96,6 +95,5 @@ void UAssetDefinitionRegistry::UnregisterAssetDefinition(UAssetDefinition* Asset
 	if (TSoftClassPtr<UObject> SupportedClass = AssetDefinition->GetAssetClass())
 	{
 		AssetDefinitions.Remove(SupportedClass);
-		AssetDefinition->OnUnregistered();
 	}
 }
\ No newline at end of file
diff --git a/Engine/Source/Editor/AssetDefinition/Public/AssetDefinition.h b/Engine/Source/Editor/AssetDefinition/Public/AssetDefinition.h
--- a/Engine/Source/Editor/AssetDefinition/Public/AssetDefinition.h
+++ b/Engine/Source/Editor/AssetDefinition/Public/AssetDefinition.h
@@ -383,11 +383,6 @@ public:
 	// 3. If you add a new function, please create a struct for the Args.  We'll be able to upgrade things easier.
 	//    Please continue to use EAssetCommandResult and FAssetSupportResponse, for those kinds of commands.
 
-public:
-
-	virtual void OnRegistered() { }
-	virtual void OnUnregistered() { }
-
 protected:
 	virtual bool CanRegisterStatically() const;
 
diff --git a/Engine/Source/Editor/UnrealEd/Public/AssetDefinitionDefault.h b/Engine/Source/Editor/UnrealEd/Public/AssetDefinitionDefault.h
--- a/Engine/Source/Editor/UnrealEd/Public/AssetDefinitionDefault.h
+++ b/Engine/Source/Editor/UnrealEd/Public/AssetDefinitionDefault.h
@@ -20,7 +20,4 @@ public:
 	virtual EAssetCommandResult GetSourceFiles(const FAssetSourceFileArgs& SourceFileArgs, TArray<FAssetSourceFile>& OutSourceAssets) const override;
 	virtual EAssetCommandResult PerformAssetDiff(const FAssetDiffArgs& DiffArgs) const override;
 	// UAssetDefinition End
-	
-public:
-	
 };
