Fixing non-unity error.

diff --git a/Engine/Plugins/Editor/EngineAssetDefinitions/Source/Private/AssetDefinition_Texture2DArray.h b/Engine/Plugins/Editor/EngineAssetDefinitions/Source/Private/AssetDefinition_Texture2DArray.h
--- a/Engine/Plugins/Editor/EngineAssetDefinitions/Source/Private/AssetDefinition_Texture2DArray.h
+++ b/Engine/Plugins/Editor/EngineAssetDefinitions/Source/Private/AssetDefinition_Texture2DArray.h
@@ -4,7 +4,7 @@
 
 #include "CoreMinimal.h"
 #include "Engine/Texture2DArray.h"
-#include "AssetDefinitionDefault.h"
+#include "AssetDefinition_Texture.h"
 
 #include "AssetDefinition_Texture2DArray.generated.h"
 
