AssetDefinition - Converting AssetTypeActions_Curve...etc to an AssetDefinition.

diff --git a/Engine/Plugins/Editor/EngineAssetDefinitions/Source/EngineAssetDefinitions.Build.cs b/Engine/Plugins/Editor/EngineAssetDefinitions/Source/EngineAssetDefinitions.Build.cs
--- a/Engine/Plugins/Editor/EngineAssetDefinitions/Source/EngineAssetDefinitions.Build.cs
+++ b/Engine/Plugins/Editor/EngineAssetDefinitions/Source/EngineAssetDefinitions.Build.cs
@@ -31,18 +31,21 @@ public class EngineAssetDefinitions : ModuleRules
 				"Kismet",
 				"TextureEditor",
 				"MaterialEditor",
+				"DesktopPlatform",
 			}
 		);
 		
 		PrivateIncludePathModuleNames.AddRange(
 			new string[] {
 				"VirtualTexturingEditor",
+				"CurveAssetEditor",
 			}
 		);
 
 		DynamicallyLoadedModuleNames.AddRange(
 			new string[] {
 				"VirtualTexturingEditor",
+				"CurveAssetEditor",
 			}
 		);
 	}
diff --git a/Engine/Plugins/Editor/EngineAssetDefinitions/Source/Private/AssetDefinition_Curve.cpp b/Engine/Plugins/Editor/EngineAssetDefinitions/Source/Private/AssetDefinition_Curve.cpp
new file mode 100644
--- /dev/null
+++ b/Engine/Plugins/Editor/EngineAssetDefinitions/Source/Private/AssetDefinition_Curve.cpp
@@ -0,0 +1,186 @@
+// Copyright Epic Games, Inc. All Rights Reserved.
+
+#include "AssetDefinition_Curve.h"
+
+#include "ContentBrowserMenuContexts.h"
+#include "EditorFramework/AssetImportData.h"
+#include "ICurveAssetEditor.h"
+#include "CurveAssetEditorModule.h"
+#include "DesktopPlatformModule.h"
+#include "IAssetTools.h"
+#include "Framework/Application/SlateApplication.h"
+#include "ToolMenuSection.h"
+#include "Misc/FileHelper.h"
+#include "Misc/MessageDialog.h"
+#include "Settings/EditorLoadingSavingSettings.h"
+#include "ToolMenus.h"
+
+#define LOCTEXT_NAMESPACE "UAssetDefinition_Curve"
+
+EAssetCommandResult UAssetDefinition_Curve::OpenAssets(const FAssetOpenArgs& OpenArgs) const
+{
+	FCurveAssetEditorModule& CurveAssetEditorModule = FModuleManager::LoadModuleChecked<FCurveAssetEditorModule>("CurveAssetEditor");
+
+	for (UCurveBase* Curve : OpenArgs.LoadObjects<UCurveBase>())
+	{
+		TSharedRef<ICurveAssetEditor> NewCurveAssetEditor = CurveAssetEditorModule.CreateCurveAssetEditor(OpenArgs.GetToolkitMode(), OpenArgs.ToolkitHost, Curve);
+	}
+
+	return EAssetCommandResult::Handled;
+}
+
+EAssetCommandResult UAssetDefinition_Curve::PerformAssetDiff(const FAssetDiffArgs& DiffArgs) const
+{
+	UCurveBase* OldCurve = Cast<UCurveBase>(DiffArgs.OldAsset);
+	UCurveBase* NewCurve = Cast<UCurveBase>(DiffArgs.NewAsset);
+
+	if (!ensure(OldCurve != nullptr && NewCurve != nullptr))
+	{
+		return EAssetCommandResult::Unhandled;
+	}
+
+	// Build names for temp json files
+	const FString RelOldTempFileName = FString::Printf(TEXT("%sTemp%s-%s.json"), *FPaths::DiffDir(), *DiffArgs.OldAsset->GetName(), *DiffArgs.OldRevision.Revision);
+	const FString AbsoluteOldTempFileName = FPaths::ConvertRelativePathToFull(RelOldTempFileName);
+	const FString RelNewTempFileName = FString::Printf(TEXT("%sTemp%s-%s.json"), *FPaths::DiffDir(), *DiffArgs.NewAsset->GetName(), *DiffArgs.NewRevision.Revision);
+	const FString AbsoluteNewTempFileName = FPaths::ConvertRelativePathToFull(RelNewTempFileName);
+
+	// save temp files
+	const bool OldResult = FFileHelper::SaveStringToFile(OldCurve->ExportAsJSONString(), *AbsoluteOldTempFileName);
+	const bool NewResult = FFileHelper::SaveStringToFile(NewCurve->ExportAsJSONString(), *AbsoluteNewTempFileName);
+
+	if (OldResult && NewResult)
+	{
+		const FString DiffCommand = GetDefault<UEditorLoadingSavingSettings>()->TextDiffToolPath.FilePath;
+		IAssetTools::Get().CreateDiffProcess(DiffCommand, AbsoluteOldTempFileName, AbsoluteNewTempFileName);
+
+		EAssetCommandResult::Handled;
+	}
+
+	return Super::PerformAssetDiff(DiffArgs);
+}
+
+// Menu Extensions
+//--------------------------------------------------------------------
+
+namespace MenuExtension_CurveBase
+{
+	static void ExecuteImportFromJSON(const FToolMenuContext& InContext)
+	{
+		IDesktopPlatform* DesktopPlatform = FDesktopPlatformModule::Get();
+
+		const void* ParentWindowWindowHandle = FSlateApplication::Get().FindBestParentWindowHandleForDialogs(nullptr);
+
+		if (const UContentBrowserAssetContextMenuContext* CBContext = UContentBrowserAssetContextMenuContext::FindContextWithAssets(InContext))
+		{
+			for (UCurveBase* Curve : CBContext->LoadSelectedObjects<UCurveBase>())
+			{
+				const FText Title = FText::Format(LOCTEXT("Curve_ImportJSONDialogTitle", "Import '{0}' from JSON..."), FText::FromString(*Curve->GetName()));
+				const FString CurrentFilename = Curve->AssetImportData->GetFirstFilename();
+				const FString FileTypes = TEXT("Curve JSON (*.json)|*.json");
+
+				TArray<FString> OutFilenames;
+				DesktopPlatform->OpenFileDialog(
+					ParentWindowWindowHandle,
+					Title.ToString(),
+					(CurrentFilename.IsEmpty()) ? TEXT("") : FPaths::GetPath(CurrentFilename),
+					(CurrentFilename.IsEmpty()) ? TEXT("") : FPaths::GetBaseFilename(CurrentFilename) + TEXT(".json"),
+					FileTypes,
+					EFileDialogFlags::None,
+					OutFilenames
+					);
+
+				if (OutFilenames.Num() > 0)
+				{
+					FString FileContent;
+					FFileHelper::LoadFileToString(FileContent, *OutFilenames[0]);
+
+					// If the JSON file contains the quoted string format, parse it into a valid JSON string, else use it as is
+					FString JSONString;
+					if (!FParse::QuotedString(*FileContent, JSONString))
+					{
+						JSONString = FileContent;
+					}
+
+					TArray<FString> Problems;
+					Curve->ImportFromJSONString(JSONString, Problems);
+					if (Problems.Num() > 0)
+					{
+						FString BigProblem;
+						for (const FString& Problem :  Problems)
+						{
+							BigProblem.Append(Problem).Append(FString(TEXT("\n")));
+						}
+						FText WarningText = FText::FromString(FString::Format(TEXT("Couldn't import file {0} into {1}:\n\n{2}"), { OutFilenames[0], Curve->GetName(), BigProblem }));
+						FMessageDialog::Open(EAppMsgType::Ok, WarningText);
+					}
+				}
+			}
+		}
+	}
+
+	static void ExecuteExportAsJSON(const FToolMenuContext& InContext)
+	{
+		IDesktopPlatform* DesktopPlatform = FDesktopPlatformModule::Get();
+
+		const void* ParentWindowWindowHandle = FSlateApplication::Get().FindBestParentWindowHandleForDialogs(nullptr);
+
+		if (const UContentBrowserAssetContextMenuContext* CBContext = UContentBrowserAssetContextMenuContext::FindContextWithAssets(InContext))
+		{
+			for (UCurveBase* Curve : CBContext->LoadSelectedObjects<UCurveBase>())
+			{
+				const FText Title = FText::Format(LOCTEXT("Curve_ExportJSONDialogTitle", "Export '{0}' as JSON..."), FText::FromString(*Curve->GetName()));
+				const FString CurrentFilename = Curve->AssetImportData->GetFirstFilename();
+				const FString FileTypes = TEXT("Curve JSON (*.json)|*.json");
+
+				TArray<FString> OutFilenames;
+				DesktopPlatform->SaveFileDialog(
+					ParentWindowWindowHandle,
+					Title.ToString(),
+					(CurrentFilename.IsEmpty()) ? TEXT("") : FPaths::GetPath(CurrentFilename),
+					(CurrentFilename.IsEmpty()) ? TEXT("") : FPaths::GetBaseFilename(CurrentFilename) + TEXT(".json"),
+					FileTypes,
+					EFileDialogFlags::None,
+					OutFilenames
+					);
+
+				if (OutFilenames.Num() > 0)
+				{
+					FFileHelper::SaveStringToFile(Curve->ExportAsJSONString(), *OutFilenames[0]);
+				}
+			}
+		}
+	}
+
+	static FDelayedAutoRegisterHelper DelayedAutoRegister(EDelayedRegisterRunPhase::EndOfEngineInit, []{ 
+		UToolMenus::RegisterStartupCallback(FSimpleMulticastDelegate::FDelegate::CreateLambda([]()
+		{
+			FToolMenuOwnerScoped OwnerScoped(UE_MODULE_NAME);
+			UToolMenu* Menu = UE::ContentBrowser::ExtendToolMenu_AssetContextMenu(UCurveBase::StaticClass());
+		
+			FToolMenuSection& Section = Menu->FindOrAddSection("GetAssetActions");
+				Section.AddDynamicEntry("GetAssetActions_UCurveBase", FNewToolMenuSectionDelegate::CreateLambda([](FToolMenuSection& InSection)
+				{
+					if (const UContentBrowserAssetContextMenuContext* Context = UContentBrowserAssetContextMenuContext::FindContextWithAssets(InSection))
+					{
+						{
+							const TAttribute<FText> Label = LOCTEXT("Curve_ImportFromJSON", "Import from JSON");
+							const TAttribute<FText> ToolTip = LOCTEXT("Curve_ImportFromJSONTooltip", "Import the curve from a file containing JSON data.");
+							const FSlateIcon Icon = FSlateIcon();
+							const FToolMenuExecuteAction UIAction = FToolMenuExecuteAction::CreateStatic(&ExecuteImportFromJSON);
+							InSection.AddMenuEntry("Curve_ImportFromJSON", Label, ToolTip, Icon, UIAction);
+						}
+						{
+							const TAttribute<FText> Label = LOCTEXT("Curve_ExportAsJSON", "Export as JSON");
+							const TAttribute<FText> ToolTip = LOCTEXT("Curve_ExportAsJSONTooltip", "Export the curve as a file containing JSON data.");
+							const FSlateIcon Icon = FSlateIcon();
+							const FToolMenuExecuteAction UIAction = FToolMenuExecuteAction::CreateStatic(&ExecuteExportAsJSON);
+							InSection.AddMenuEntry("Curve_ExportAsJSON", Label, ToolTip, Icon, UIAction);
+						}
+					}
+				}));
+		}));
+	});
+}
+
+#undef LOCTEXT_NAMESPACE
diff --git a/Engine/Plugins/Editor/EngineAssetDefinitions/Source/Private/AssetDefinition_Curve.h b/Engine/Plugins/Editor/EngineAssetDefinitions/Source/Private/AssetDefinition_Curve.h
new file mode 100644
--- /dev/null
+++ b/Engine/Plugins/Editor/EngineAssetDefinitions/Source/Private/AssetDefinition_Curve.h
@@ -0,0 +1,27 @@
+// Copyright Epic Games, Inc. All Rights Reserved.
+
+#pragma once
+
+#include "CoreMinimal.h"
+#include "Toolkits/IToolkitHost.h"
+#include "AssetDefinitionDefault.h"
+#include "Curves/CurveBase.h"
+
+#include "AssetDefinition_Curve.generated.h"
+
+UCLASS()
+class UAssetDefinition_Curve : public UAssetDefinitionDefault
+{
+	GENERATED_BODY()
+
+public:
+	// UAssetDefinition Begin
+	virtual FText GetAssetDisplayName() const override { return NSLOCTEXT("AssetTypeActions", "AssetTypeActions_Curve", "Curve"); }
+	virtual FLinearColor GetAssetColor() const override { return FLinearColor(FColor(78, 40, 165)); }
+	virtual TSoftClassPtr<UObject> GetAssetClass() const override { return UCurveBase::StaticClass(); }
+	virtual bool CanImport() const override { return true; }
+
+	virtual EAssetCommandResult OpenAssets(const FAssetOpenArgs& OpenArgs) const override;
+	virtual EAssetCommandResult PerformAssetDiff(const FAssetDiffArgs& DiffArgs) const override;
+	// UAssetDefinition End
+};
diff --git a/Engine/Plugins/Editor/EngineAssetDefinitions/Source/Private/AssetDefinition_CurveFloat.h b/Engine/Plugins/Editor/EngineAssetDefinitions/Source/Private/AssetDefinition_CurveFloat.h
new file mode 100644
--- /dev/null
+++ b/Engine/Plugins/Editor/EngineAssetDefinitions/Source/Private/AssetDefinition_CurveFloat.h
@@ -0,0 +1,21 @@
+// Copyright Epic Games, Inc. All Rights Reserved.
+
+#pragma once
+
+#include "CoreMinimal.h"
+#include "AssetDefinition_Curve.h"
+#include "Curves/CurveFloat.h"
+
+#include "AssetDefinition_CurveFloat.generated.h"
+
+UCLASS()
+class UAssetDefinition_CurveFloat : public UAssetDefinition_Curve
+{
+	GENERATED_BODY()
+
+public:
+	// UAssetDefinition Begin
+	virtual FText GetAssetDisplayName() const override { return NSLOCTEXT("AssetTypeActions", "AssetTypeActions_CurveFloat", "Float Curve"); }
+	virtual TSoftClassPtr<UObject> GetAssetClass() const override { return UCurveFloat::StaticClass(); }
+	// UAssetDefinition End
+};
diff --git a/Engine/Plugins/Editor/EngineAssetDefinitions/Source/Private/AssetDefinition_CurveLinearColor.h b/Engine/Plugins/Editor/EngineAssetDefinitions/Source/Private/AssetDefinition_CurveLinearColor.h
new file mode 100644
--- /dev/null
+++ b/Engine/Plugins/Editor/EngineAssetDefinitions/Source/Private/AssetDefinition_CurveLinearColor.h
@@ -0,0 +1,21 @@
+// Copyright Epic Games, Inc. All Rights Reserved.
+
+#pragma once
+
+#include "CoreMinimal.h"
+#include "AssetDefinition_Curve.h"
+#include "Curves/CurveLinearColor.h"
+
+#include "AssetDefinition_CurveLinearColor.generated.h"
+
+UCLASS()
+class UAssetDefinition_CurveLinearColor : public UAssetDefinition_Curve
+{
+	GENERATED_BODY()
+
+public:
+	// UAssetDefinition Begin
+	virtual FText GetAssetDisplayName() const override { return NSLOCTEXT("AssetTypeActions", "AssetTypeActions_CurveLinearColor", "Color Curve"); }
+	virtual TSoftClassPtr<UObject> GetAssetClass() const override { return UCurveLinearColor::StaticClass(); }
+	// UAssetDefinition End
+};
diff --git a/Engine/Plugins/Editor/EngineAssetDefinitions/Source/Private/AssetDefinition_CurveLinearColorAtlas.h b/Engine/Plugins/Editor/EngineAssetDefinitions/Source/Private/AssetDefinition_CurveLinearColorAtlas.h
--- a/Engine/Plugins/Editor/EngineAssetDefinitions/Source/Private/AssetDefinition_CurveLinearColorAtlas.h
+++ b/Engine/Plugins/Editor/EngineAssetDefinitions/Source/Private/AssetDefinition_CurveLinearColorAtlas.h
@@ -9,7 +9,7 @@
 #include "AssetDefinition_CurveLinearColorAtlas.generated.h"
 
 UCLASS()
-class ENGINEASSETDEFINITIONS_API UAssetDefinition_CurveLinearColorAtlas : public UAssetDefinition_Texture
+class UAssetDefinition_CurveLinearColorAtlas : public UAssetDefinition_Texture
 {
 	GENERATED_BODY()
 
diff --git a/Engine/Plugins/Editor/EngineAssetDefinitions/Source/Private/AssetDefinition_CurveVector.h b/Engine/Plugins/Editor/EngineAssetDefinitions/Source/Private/AssetDefinition_CurveVector.h
new file mode 100644
--- /dev/null
+++ b/Engine/Plugins/Editor/EngineAssetDefinitions/Source/Private/AssetDefinition_CurveVector.h
@@ -0,0 +1,21 @@
+// Copyright Epic Games, Inc. All Rights Reserved.
+
+#pragma once
+
+#include "CoreMinimal.h"
+#include "AssetDefinition_Curve.h"
+#include "Curves/CurveVector.h"
+
+#include "AssetDefinition_CurveVector.generated.h"
+
+UCLASS()
+class UAssetDefinition_CurveVector : public UAssetDefinition_Curve
+{
+	GENERATED_BODY()
+
+public:
+	// UAssetDefinition Begin
+	virtual FText GetAssetDisplayName() const override { return NSLOCTEXT("AssetTypeActions", "AssetTypeActions_CurveVector", "Vector Curve"); }
+	virtual TSoftClassPtr<UObject> GetAssetClass() const override { return UCurveVector::StaticClass(); }
+	// UAssetDefinition End
+};
diff --git a/Engine/Plugins/Editor/EngineAssetDefinitions/Source/Private/AssetDefinition_Texture.h b/Engine/Plugins/Editor/EngineAssetDefinitions/Source/Private/AssetDefinition_Texture.h
--- a/Engine/Plugins/Editor/EngineAssetDefinitions/Source/Private/AssetDefinition_Texture.h
+++ b/Engine/Plugins/Editor/EngineAssetDefinitions/Source/Private/AssetDefinition_Texture.h
@@ -8,8 +8,8 @@
 
 #include "AssetDefinition_Texture.generated.h"
 
-UCLASS(Abstract)
-class ENGINEASSETDEFINITIONS_API UAssetDefinition_Texture : public UAssetDefinitionDefault
+UCLASS()
+class UAssetDefinition_Texture : public UAssetDefinitionDefault
 {
 	GENERATED_BODY()
 
diff --git a/Engine/Plugins/Editor/EngineAssetDefinitions/Source/Private/AssetDefinition_UserDefinedStruct.h b/Engine/Plugins/Editor/EngineAssetDefinitions/Source/Private/AssetDefinition_UserDefinedStruct.h
--- a/Engine/Plugins/Editor/EngineAssetDefinitions/Source/Private/AssetDefinition_UserDefinedStruct.h
+++ b/Engine/Plugins/Editor/EngineAssetDefinitions/Source/Private/AssetDefinition_UserDefinedStruct.h
@@ -9,7 +9,7 @@
 #include "AssetDefinition_UserDefinedStruct.generated.h"
 
 UCLASS()
-class ENGINEASSETDEFINITIONS_API UAssetDefinition_UserDefinedStruct : public UAssetDefinitionDefault
+class UAssetDefinition_UserDefinedStruct : public UAssetDefinitionDefault
 {
 	GENERATED_BODY()
 
diff --git a/Engine/Plugins/Editor/EngineAssetDefinitions/Source/Public/AssetDefinition_Blueprint.h b/Engine/Plugins/Editor/EngineAssetDefinitions/Source/Public/AssetDefinition_Blueprint.h
--- a/Engine/Plugins/Editor/EngineAssetDefinitions/Source/Public/AssetDefinition_Blueprint.h
+++ b/Engine/Plugins/Editor/EngineAssetDefinitions/Source/Public/AssetDefinition_Blueprint.h
@@ -13,7 +13,7 @@ struct FAssetData;
 class UFactory;
 
 UCLASS()
-class ENGINEASSETDEFINITIONS_API UAssetDefinition_Blueprint : public UAssetDefinitionDefault
+class UAssetDefinition_Blueprint : public UAssetDefinitionDefault
 {
 	GENERATED_BODY()
 
diff --git a/Engine/Plugins/Editor/EngineAssetDefinitions/Source/Public/AssetDefinition_DataAsset.h b/Engine/Plugins/Editor/EngineAssetDefinitions/Source/Public/AssetDefinition_DataAsset.h
--- a/Engine/Plugins/Editor/EngineAssetDefinitions/Source/Public/AssetDefinition_DataAsset.h
+++ b/Engine/Plugins/Editor/EngineAssetDefinitions/Source/Public/AssetDefinition_DataAsset.h
@@ -9,7 +9,7 @@
 #include "AssetDefinition_DataAsset.generated.h"
 
 UCLASS()
-class ENGINEASSETDEFINITIONS_API UAssetDefinition_DataAsset : public UAssetDefinitionDefault
+class UAssetDefinition_DataAsset : public UAssetDefinitionDefault
 {
 	GENERATED_BODY()
 
diff --git a/Engine/Source/Developer/AssetTools/Private/AssetTools.cpp b/Engine/Source/Developer/AssetTools/Private/AssetTools.cpp
--- a/Engine/Source/Developer/AssetTools/Private/AssetTools.cpp
+++ b/Engine/Source/Developer/AssetTools/Private/AssetTools.cpp
@@ -39,7 +39,6 @@
 #include "AssetTypeActions/AssetTypeActions_ActorFolder.h"
 #include "AssetTypeActions/AssetTypeActions_Blueprint.h"
 #include "AssetTypeActions/AssetTypeActions_BlueprintGeneratedClass.h"
-#include "AssetTypeActions/AssetTypeActions_Curve.h"
 #include "AssetTypeActions/AssetTypeActions_MaterialInterface.h"
 #include "AssetTypeActions/AssetTypeActions_SkeletalMesh.h"
 #include "AssetTypeActions/AssetTypeActions_FbxSceneImportData.h"
@@ -58,11 +57,8 @@
 #include "AssetTypeActions/AssetTypeActions_AimOffset.h"
 #include "AssetTypeActions/AssetTypeActions_BlendSpace1D.h"
 #include "AssetTypeActions/AssetTypeActions_AimOffset1D.h"
-#include "AssetTypeActions/AssetTypeActions_CurveFloat.h"
 #include "AssetTypeActions/AssetTypeActions_CurveTable.h"
 #include "AssetTypeActions/AssetTypeActions_CompositeCurveTable.h"
-#include "AssetTypeActions/AssetTypeActions_CurveVector.h"
-#include "AssetTypeActions/AssetTypeActions_CurveLinearColor.h"
 #include "AssetTypeActions/AssetTypeActions_DataAsset.h"
 #include "AssetTypeActions/AssetTypeActions_DataLayer.h"
 #include "AssetTypeActions/AssetTypeActions_DataTable.h"
@@ -1248,12 +1244,8 @@ UAssetToolsImpl::UAssetToolsImpl(const FObjectInitializer& ObjectInitializer)
 	RegisterAssetTypeActions(MakeShareable(new FAssetTypeActions_BlendSpace1D));
 	RegisterAssetTypeActions(MakeShareable(new FAssetTypeActions_Blueprint));
 	RegisterAssetTypeActions(MakeShareable(new FAssetTypeActions_BlueprintGeneratedClass));
-	RegisterAssetTypeActions(MakeShareable(new FAssetTypeActions_Curve));
-	RegisterAssetTypeActions(MakeShareable(new FAssetTypeActions_CurveFloat));
 	RegisterAssetTypeActions(MakeShareable(new FAssetTypeActions_CurveTable));
 	RegisterAssetTypeActions(MakeShareable(new FAssetTypeActions_CompositeCurveTable));
-	RegisterAssetTypeActions(MakeShareable(new FAssetTypeActions_CurveVector));
-	RegisterAssetTypeActions(MakeShareable(new FAssetTypeActions_CurveLinearColor));
 	RegisterAssetTypeActions(MakeShareable(new FAssetTypeActions_DataLayer));
 	RegisterAssetTypeActions(MakeShareable(new FAssetTypeActions_DataTable));
 	RegisterAssetTypeActions(MakeShareable(new FAssetTypeActions_CompositeDataTable));
diff --git a/Engine/Source/Developer/AssetTools/Private/AssetTypeActions/AssetTypeActions_Curve.cpp b/Engine/Source/Developer/AssetTools/Private/AssetTypeActions/AssetTypeActions_Curve.cpp
deleted file mode 100644
diff --git a/Engine/Source/Developer/AssetTools/Private/AssetTypeActions/AssetTypeActions_Curve.h b/Engine/Source/Developer/AssetTools/Private/AssetTypeActions/AssetTypeActions_Curve.h
deleted file mode 100644
diff --git a/Engine/Source/Developer/AssetTools/Private/AssetTypeActions/AssetTypeActions_CurveFloat.h b/Engine/Source/Developer/AssetTools/Private/AssetTypeActions/AssetTypeActions_CurveFloat.h
deleted file mode 100644
diff --git a/Engine/Source/Developer/AssetTools/Private/AssetTypeActions/AssetTypeActions_CurveLinearColor.h b/Engine/Source/Developer/AssetTools/Private/AssetTypeActions/AssetTypeActions_CurveLinearColor.h
deleted file mode 100644
diff --git a/Engine/Source/Developer/AssetTools/Private/AssetTypeActions/AssetTypeActions_CurveVector.h b/Engine/Source/Developer/AssetTools/Private/AssetTypeActions/AssetTypeActions_CurveVector.h
deleted file mode 100644
diff --git a/Engine/Source/Editor/AssetDefinition/Private/AssetDefinition.cpp b/Engine/Source/Editor/AssetDefinition/Private/AssetDefinition.cpp
--- a/Engine/Source/Editor/AssetDefinition/Private/AssetDefinition.cpp
+++ b/Engine/Source/Editor/AssetDefinition/Private/AssetDefinition.cpp
@@ -53,6 +53,12 @@ void UAssetDefinition::PostCDOContruct()
 	}
 }
 
+TConstArrayView<FAssetCategoryPath> UAssetDefinition::GetAssetCategories() const
+{
+	static const auto Categories = { EAssetCategoryPaths::Misc };
+	return Categories;
+}
+
 bool UAssetDefinition::CanRegisterStatically() const
 {
 	return !GetClass()->HasAnyClassFlags(CLASS_Abstract);
diff --git a/Engine/Source/Editor/AssetDefinition/Public/AssetDefinition.h b/Engine/Source/Editor/AssetDefinition/Public/AssetDefinition.h
--- a/Engine/Source/Editor/AssetDefinition/Public/AssetDefinition.h
+++ b/Engine/Source/Editor/AssetDefinition/Public/AssetDefinition.h
@@ -289,11 +289,7 @@ public:
 	virtual FText GetAssetDescription(const FAssetData& AssetData) const { return FText::GetEmpty(); }
 
 	/** Gets a list of categories this asset is in, these categories are used to help organize */
-	virtual TConstArrayView<FAssetCategoryPath> GetAssetCategories() const
-	{
-		static const auto Categories = { EAssetCategoryPaths::Misc };
-		return Categories;
-	}
+	virtual TConstArrayView<FAssetCategoryPath> GetAssetCategories() const;
 	
 public:
 	// Common Operations
