Fixing linux build error.

diff --git a/Engine/Plugins/Editor/EngineAssetDefinitions/Source/Private/AssetDefinition_Curve.cpp b/Engine/Plugins/Editor/EngineAssetDefinitions/Source/Private/AssetDefinition_Curve.cpp
--- a/Engine/Plugins/Editor/EngineAssetDefinitions/Source/Private/AssetDefinition_Curve.cpp
+++ b/Engine/Plugins/Editor/EngineAssetDefinitions/Source/Private/AssetDefinition_Curve.cpp
@@ -54,7 +54,7 @@ EAssetCommandResult UAssetDefinition_Curve::PerformAssetDiff(const FAssetDiffArg
 		const FString DiffCommand = GetDefault<UEditorLoadingSavingSettings>()->TextDiffToolPath.FilePath;
 		IAssetTools::Get().CreateDiffProcess(DiffCommand, AbsoluteOldTempFileName, AbsoluteNewTempFileName);
 
-		EAssetCommandResult::Handled;
+		return EAssetCommandResult::Handled;
 	}
 
 	return Super::PerformAssetDiff(DiffArgs);
