AssetDefinitionStaticMesh - Fixing a couple of menu differences, with copy LOD and add new LOD.

diff --git a/Engine/Plugins/Editor/EngineAssetDefinitions/Source/Private/AssetDefinition_StaticMesh.cpp b/Engine/Plugins/Editor/EngineAssetDefinitions/Source/Private/AssetDefinition_StaticMesh.cpp
--- a/Engine/Plugins/Editor/EngineAssetDefinitions/Source/Private/AssetDefinition_StaticMesh.cpp
+++ b/Engine/Plugins/Editor/EngineAssetDefinitions/Source/Private/AssetDefinition_StaticMesh.cpp
@@ -382,13 +382,14 @@ namespace MenuExtension_StaticMesh
 		const FAssetData& StaticMeshAsset = Context->SelectedAssets[0];
 
 		//Add 1 so we can Add LOD when we are not in "Reimport With New File" submenu
-    	const int32 LodCount = GetNumberOfLODs(StaticMeshAsset) + (bReimportWithNewFile ? 0 : 1);
+    	const int32 OriginalLodCount = GetNumberOfLODs(StaticMeshAsset);
+    	const int32 LodCount = OriginalLodCount+ (bReimportWithNewFile ? 0 : 1);
     	
 		for (int32 LOD = 1; LOD < LodCount; ++LOD)
 		{
 			FText Label = FText::Format(LOCTEXT("Reimport LOD (number)", "Reimport LOD {0}"), LOD);
 			FText ToolTip = LOCTEXT("ReimportTip", "Reimport over existing LOD");
-			if (LOD == LodCount)
+			if (LOD == OriginalLodCount)
 			{
 				Label = FText::Format(LOCTEXT("Add LOD (number)", "Add LOD {0}"), LOD);
 				ToolTip = LOCTEXT("NewImportTip", "Import new LOD");
@@ -451,7 +452,7 @@ namespace MenuExtension_StaticMesh
 			FToolUIAction UIAction;
 			UIAction.ExecuteAction = FToolMenuExecuteAction::CreateStatic(&ExecutePasteLODSettings);
 			UIAction.CanExecuteAction = FToolMenuCanExecuteAction::CreateStatic(&CanPasteLODSettings);
-			Section.AddMenuEntry("StaticMesh_CopyLOD", Label, ToolTip, Icon, UIAction);
+			Section.AddMenuEntry("StaticMesh_PasteLOD", Label, ToolTip, Icon, UIAction);
 	    }
  
     	if (Context->SelectedAssets.Num() == 1)
