Fixing non-unity build error.

diff --git a/Engine/Plugins/Editor/EngineAssetDefinitions/Source/Private/AssetDefinition_VectorFieldAnimated.h b/Engine/Plugins/Editor/EngineAssetDefinitions/Source/Private/AssetDefinition_VectorFieldAnimated.h
--- a/Engine/Plugins/Editor/EngineAssetDefinitions/Source/Private/AssetDefinition_VectorFieldAnimated.h
+++ b/Engine/Plugins/Editor/EngineAssetDefinitions/Source/Private/AssetDefinition_VectorFieldAnimated.h
@@ -4,7 +4,7 @@
 
 #include "CoreMinimal.h"
 #include "VectorField/VectorFieldAnimated.h"
-#include "AssetDefinitionDefault.h"
+#include "AssetDefinition_VectorField.h"
 
 #include "AssetDefinition_VectorFieldAnimated.generated.h"
 
