Fixing non unity for other vector field type.

diff --git a/Engine/Plugins/Editor/EngineAssetDefinitions/Source/Private/AssetDefinition_VectorFieldStatic.h b/Engine/Plugins/Editor/EngineAssetDefinitions/Source/Private/AssetDefinition_VectorFieldStatic.h
--- a/Engine/Plugins/Editor/EngineAssetDefinitions/Source/Private/AssetDefinition_VectorFieldStatic.h
+++ b/Engine/Plugins/Editor/EngineAssetDefinitions/Source/Private/AssetDefinition_VectorFieldStatic.h
@@ -4,7 +4,7 @@
 
 #include "CoreMinimal.h"
 #include "VectorField/VectorFieldStatic.h"
-#include "AssetDefinitionDefault.h"
+#include "AssetDefinition_VectorField.h"
 
 #include "AssetDefinition_VectorFieldStatic.generated.h"
 
