Make Virtual Texture Convert/Deconvert menus work again.

diff --git a/Engine/Plugins/Editor/EngineAssetDefinitions/Source/Private/AssetDefinition_Texture.cpp b/Engine/Plugins/Editor/EngineAssetDefinitions/Source/Private/AssetDefinition_Texture.cpp
--- a/Engine/Plugins/Editor/EngineAssetDefinitions/Source/Private/AssetDefinition_Texture.cpp
+++ b/Engine/Plugins/Editor/EngineAssetDefinitions/Source/Private/AssetDefinition_Texture.cpp
@@ -213,8 +213,17 @@ namespace MenuExtension_Texture
 								AssetData.GetTagValue<bool>("VirtualTextureStreaming", VirtualTextured);
 								return VirtualTextured;
 							});
-					
+
 						if (bHasVirtualTextures)
+						{
+							const TAttribute<FText> Label = LOCTEXT("Texture_ConvertToRegular", "Convert to Regular Texture");
+							const TAttribute<FText> ToolTip = LOCTEXT("Texture_ConvertToRegularTooltip", "Converts this texture to a regular 2D texture if it is a virtual texture.");
+							const FSlateIcon Icon = FSlateIcon(FAppStyle::GetAppStyleSetName(), "ClassIcon.Texture2D");
+							const FToolMenuExecuteAction UIAction = FToolMenuExecuteAction::CreateStatic(&ExecuteConvertToRegularTexture);
+							InSection.AddMenuEntry("Texture_ConvertToVT", Label, ToolTip, Icon, UIAction);
+						}
+						
+						if (!bHasVirtualTextures)
 						{
 							const TAttribute<FText> Label = LOCTEXT("Texture_ConvertToVT", "Convert to Virtual Texture");
 							const TAttribute<FText> ToolTip = LOCTEXT("Texture_ConvertToVTTooltip", "Converts this texture to a virtual texture if it fits the size limit imposed in the texture importer settings.");
@@ -222,6 +231,7 @@ namespace MenuExtension_Texture
 							const FToolMenuExecuteAction UIAction = FToolMenuExecuteAction::CreateStatic(&ExecuteConvertToVirtualTexture);
 							InSection.AddMenuEntry("Texture_ConvertToVT", Label, ToolTip, Icon, UIAction);
 						}
+
 						
 						if ( Context->SelectedAssets.Num() == 1 )
 						{
