Temp workaround to bring back the Data Asset filter in the Content Browser

diff --git a/Engine/Plugins/Editor/EngineAssetDefinitions/Source/Private/AssetDefinition_DataAsset.cpp b/Engine/Plugins/Editor/EngineAssetDefinitions/Source/Private/AssetDefinition_DataAsset.cpp
--- a/Engine/Plugins/Editor/EngineAssetDefinitions/Source/Private/AssetDefinition_DataAsset.cpp
+++ b/Engine/Plugins/Editor/EngineAssetDefinitions/Source/Private/AssetDefinition_DataAsset.cpp
@@ -11,6 +11,7 @@
 #include "ToolMenus.h"
 #include "Kismet2/SClassPickerDialog.h"
 #include "Engine/Engine.h"
+#include "Misc/AssetFilterData.h"
 
 #define LOCTEXT_NAMESPACE "UAssetDefinition_DataAsset"
 
@@ -144,4 +145,25 @@ EAssetCommandResult UAssetDefinition_DataAsset::PerformAssetDiff(const FAssetDif
 	return EAssetCommandResult::Handled;
 }
 
+EAssetCommandResult UAssetDefinition_DataAsset::GetFilters(TArray<FAssetFilterData>& OutFilters) const
+{
+	// AssetDefinitions don't allow filters for abstract classes by default, so we override that behavior to allow
+	// filters for Data Assets for now
+	const TSoftClassPtr<UObject> AssetClassPtr = GetAssetClass();
+	
+	if (const UClass* AssetClass = AssetClassPtr.Get())
+	{
+		FAssetFilterData DefaultFilter;
+		DefaultFilter.Name = AssetClassPtr.ToSoftObjectPath().ToString();
+		DefaultFilter.DisplayText = GetAssetDisplayName();
+		DefaultFilter.Filter.ClassPaths.Add(AssetClassPtr.ToSoftObjectPath().GetAssetPath());
+		DefaultFilter.Filter.bRecursiveClasses = true;
+		OutFilters.Add(MoveTemp(DefaultFilter));
+
+		return EAssetCommandResult::Handled;
+	}
+	
+	return EAssetCommandResult::Unhandled;
+}
+
 #undef LOCTEXT_NAMESPACE
diff --git a/Engine/Plugins/Editor/EngineAssetDefinitions/Source/Public/AssetDefinition_DataAsset.h b/Engine/Plugins/Editor/EngineAssetDefinitions/Source/Public/AssetDefinition_DataAsset.h
--- a/Engine/Plugins/Editor/EngineAssetDefinitions/Source/Public/AssetDefinition_DataAsset.h
+++ b/Engine/Plugins/Editor/EngineAssetDefinitions/Source/Public/AssetDefinition_DataAsset.h
@@ -21,4 +21,5 @@ public:
 	virtual TSoftClassPtr<UObject> GetAssetClass() const override { return UDataAsset::StaticClass(); }
 	
 	virtual EAssetCommandResult PerformAssetDiff(const FAssetDiffArgs& DiffArgs) const override;
+	virtual EAssetCommandResult GetFilters(TArray<FAssetFilterData>& OutFilters) const override;
 };
