#UE Add AssetTypeActions to update source animation on poses. This class should be converted to an assetdefinition instead, but that can be a different change. Also added the logic to properly register this assettypeactions

diff --git a/Engine/Plugins/Editor/EngineAssetDefinitions/Source/Private/AssetDefinition_PoseAsset.cpp b/Engine/Plugins/Editor/EngineAssetDefinitions/Source/Private/AssetDefinition_PoseAsset.cpp
new file mode 100644
--- /dev/null
+++ b/Engine/Plugins/Editor/EngineAssetDefinitions/Source/Private/AssetDefinition_PoseAsset.cpp
@@ -0,0 +1,50 @@
+// Copyright Epic Games, Inc. All Rights Reserved.
+
+#include "AssetDefinition_PoseAsset.h"
+#include "ToolMenuSection.h"
+#include "Animation/PoseAsset.h"
+#include "Misc/MessageDialog.h"
+#include "Styling/AppStyle.h"
+#include "ToolMenus.h"
+#include "Framework/Application/SlateApplication.h"
+
+#define LOCTEXT_NAMESPACE "AssetTypeActions"
+
+void FAssetTypeActions_PoseAsset::GetActions(const TArray<UObject*>& InObjects, FToolMenuSection& Section)
+{
+	auto PoseAssets = GetTypedWeakObjectPtrs<UPoseAsset>(InObjects);
+
+	Section.AddMenuEntry(
+		"PoseAsset_UpdateSource",
+		LOCTEXT("PoseAsset_UpdateSource", "Update Source Animation"),
+		LOCTEXT("PoseAsset_UpdateSourceTooltip", "Updates the source animation for this pose"),
+		FSlateIcon(),
+		FUIAction(
+			FExecuteAction::CreateSP(this, &FAssetTypeActions_PoseAsset::ExecuteUpdateSource, PoseAssets),
+			FCanExecuteAction()
+		)
+	);
+}
+void FAssetTypeActions_PoseAsset::ExecuteUpdateSource(TArray<TWeakObjectPtr<UPoseAsset>> Objects)
+{
+	FScopedTransaction Transaction(LOCTEXT("PoseUpdateSource", "Updating Source Animation for Pose"));
+	for (const TWeakObjectPtr<UPoseAsset>& WeakPose : Objects)
+	{
+		if (UPoseAsset* PoseAsset = WeakPose.Get())
+		{
+			if (UAnimSequence* SourceAnimation = PoseAsset->SourceAnimation)
+			{
+				if (!PoseAsset->SourceAnimationRawDataGUID.IsValid() || PoseAsset->SourceAnimationRawDataGUID != SourceAnimation->GetRawDataGuid())
+				{
+					if (PoseAsset->GetSkeleton()->IsCompatibleForEditor(SourceAnimation->GetSkeleton()))
+					{
+						PoseAsset->Modify();
+						PoseAsset->UpdatePoseFromAnimation(SourceAnimation);
+					}
+				}
+			}
+		}
+	}
+}
+
+#undef LOCTEXT_NAMESPACE
diff --git a/Engine/Plugins/Editor/EngineAssetDefinitions/Source/Private/AssetDefinition_PoseAsset.h b/Engine/Plugins/Editor/EngineAssetDefinitions/Source/Private/AssetDefinition_PoseAsset.h
--- a/Engine/Plugins/Editor/EngineAssetDefinitions/Source/Private/AssetDefinition_PoseAsset.h
+++ b/Engine/Plugins/Editor/EngineAssetDefinitions/Source/Private/AssetDefinition_PoseAsset.h
@@ -14,4 +14,10 @@ public:
 	virtual FColor GetTypeColor() const override { return FColor(237, 28, 36); }
 	virtual UClass* GetSupportedClass() const override { return UPoseAsset::StaticClass(); }
 	virtual bool CanFilter() override { return true; }
+
+	virtual void GetActions(const TArray<UObject*>& InObjects, struct FToolMenuSection& Section) override;
+
+private:
+
+	void ExecuteUpdateSource(TArray<TWeakObjectPtr<UPoseAsset>> Objects);
 };
diff --git a/Engine/Plugins/Editor/EngineAssetDefinitions/Source/Private/EngineAssetDefinitionsModule.cpp b/Engine/Plugins/Editor/EngineAssetDefinitions/Source/Private/EngineAssetDefinitionsModule.cpp
--- a/Engine/Plugins/Editor/EngineAssetDefinitions/Source/Private/EngineAssetDefinitionsModule.cpp
+++ b/Engine/Plugins/Editor/EngineAssetDefinitions/Source/Private/EngineAssetDefinitionsModule.cpp
@@ -1,6 +1,9 @@
 ﻿// Copyright Epic Games, Inc. All Rights Reserved.
 
 #include "Modules/ModuleManager.h"
+#include "AssetToolsModule.h"
+#include "IAssetTypeActions.h"
+#include "AssetDefinition_PoseAsset.h"
 
 #define LOCTEXT_NAMESPACE "FEngineAssetDefinitionsModule"
 
@@ -11,11 +14,33 @@ public:
 	// IModuleInterface interface
 	virtual void StartupModule() override
 	{
+		// Register asset types... for now
+		IAssetTools& AssetTools = FModuleManager::LoadModuleChecked<FAssetToolsModule>("AssetTools").Get();
+		RegisterAssetTypeAction(AssetTools, MakeShareable(new FAssetTypeActions_PoseAsset));
 	}
 
 	virtual void ShutdownModule() override
 	{
+		// Unregister all the asset types that we registered
+		if (FModuleManager::Get().IsModuleLoaded("AssetTools"))
+		{
+			IAssetTools& AssetTools = FModuleManager::GetModuleChecked<FAssetToolsModule>("AssetTools").Get();
+			for (int32 Index = 0; Index < CreatedAssetTypeActions.Num(); ++Index)
+			{
+				AssetTools.UnregisterAssetTypeActions(CreatedAssetTypeActions[Index].ToSharedRef());
+			}
+		}
+		CreatedAssetTypeActions.Empty();
 	}
+
+private:
+	void RegisterAssetTypeAction(IAssetTools& AssetTools, TSharedRef<IAssetTypeActions> Action)
+	{
+		AssetTools.RegisterAssetTypeActions(Action);
+		CreatedAssetTypeActions.Add(Action);
+	}
+
+	TArray<TSharedPtr<IAssetTypeActions>> CreatedAssetTypeActions;
 };
 
 #undef LOCTEXT_NAMESPACE
