#UE Non-unity fix

diff --git a/Engine/Plugins/Editor/EngineAssetDefinitions/Source/Private/AssetDefinition_PoseAsset.cpp b/Engine/Plugins/Editor/EngineAssetDefinitions/Source/Private/AssetDefinition_PoseAsset.cpp
--- a/Engine/Plugins/Editor/EngineAssetDefinitions/Source/Private/AssetDefinition_PoseAsset.cpp
+++ b/Engine/Plugins/Editor/EngineAssetDefinitions/Source/Private/AssetDefinition_PoseAsset.cpp
@@ -6,6 +6,7 @@
 #include "Misc/MessageDialog.h"
 #include "Styling/AppStyle.h"
 #include "ToolMenus.h"
+#include "ScopedTransaction.h"
 #include "Framework/Application/SlateApplication.h"
 
 #define LOCTEXT_NAMESPACE "AssetTypeActions"
