#UE Another non-unity fix

diff --git a/Engine/Plugins/Editor/EngineAssetDefinitions/Source/Private/AssetDefinition_PoseAsset.cpp b/Engine/Plugins/Editor/EngineAssetDefinitions/Source/Private/AssetDefinition_PoseAsset.cpp
--- a/Engine/Plugins/Editor/EngineAssetDefinitions/Source/Private/AssetDefinition_PoseAsset.cpp
+++ b/Engine/Plugins/Editor/EngineAssetDefinitions/Source/Private/AssetDefinition_PoseAsset.cpp
@@ -3,6 +3,7 @@
 #include "AssetDefinition_PoseAsset.h"
 #include "ToolMenuSection.h"
 #include "Animation/PoseAsset.h"
+#include "Animation/AnimSequence.h"
 #include "Misc/MessageDialog.h"
 #include "Styling/AppStyle.h"
 #include "ToolMenus.h"
