#UE Update AssetDefinition_PoseAsset.cpp to match the logic that is in PoseAsset.cpp to detect if the source animation needs to be updated.

diff --git a/Engine/Plugins/Editor/EngineAssetDefinitions/Source/Private/AssetDefinition_PoseAsset.cpp b/Engine/Plugins/Editor/EngineAssetDefinitions/Source/Private/AssetDefinition_PoseAsset.cpp
--- a/Engine/Plugins/Editor/EngineAssetDefinitions/Source/Private/AssetDefinition_PoseAsset.cpp
+++ b/Engine/Plugins/Editor/EngineAssetDefinitions/Source/Private/AssetDefinition_PoseAsset.cpp
@@ -36,7 +36,7 @@ void FAssetTypeActions_PoseAsset::ExecuteUpdateSource(TArray<TWeakObjectPtr<UPos
 		{
 			if (UAnimSequence* SourceAnimation = PoseAsset->SourceAnimation)
 			{
-				if (!PoseAsset->SourceAnimationRawDataGUID.IsValid() || PoseAsset->SourceAnimationRawDataGUID != SourceAnimation->GetRawDataGuid())
+				if (!PoseAsset->SourceAnimationRawDataGUID.IsValid() || PoseAsset->SourceAnimationRawDataGUID != SourceAnimation->GetDataModel()->GenerateGuid())
 				{
 					if (PoseAsset->GetSkeleton()->IsCompatibleForEditor(SourceAnimation->GetSkeleton()))
 					{
