AssetDefinition - Tweaking the MIC to allow importing, since that was added in main.

diff --git a/Engine/Plugins/Editor/EngineAssetDefinitions/Source/Private/Material/AssetDefinition_MaterialInstanceConstant.h b/Engine/Plugins/Editor/EngineAssetDefinitions/Source/Private/Material/AssetDefinition_MaterialInstanceConstant.h
--- a/Engine/Plugins/Editor/EngineAssetDefinitions/Source/Private/Material/AssetDefinition_MaterialInstanceConstant.h
+++ b/Engine/Plugins/Editor/EngineAssetDefinitions/Source/Private/Material/AssetDefinition_MaterialInstanceConstant.h
@@ -19,5 +19,6 @@ public:
 	virtual FLinearColor GetAssetColor() const override { return FLinearColor(FColor(0,128,0)); }
 	virtual TSoftClassPtr<UObject> GetAssetClass() const override { return UMaterialInstanceConstant::StaticClass(); }
 	virtual EAssetCommandResult OpenAssets(const FAssetOpenArgs& OpenArgs) const override;
+	virtual bool CanImport() const override { return true; }
 	// UAssetDefinition End
 };
diff --git a/Engine/Source/Editor/UnrealEd/Public/AssetDefinitionDefault.h b/Engine/Source/Editor/UnrealEd/Public/AssetDefinitionDefault.h
--- a/Engine/Source/Editor/UnrealEd/Public/AssetDefinitionDefault.h
+++ b/Engine/Source/Editor/UnrealEd/Public/AssetDefinitionDefault.h
@@ -3,12 +3,10 @@
 #pragma once
 
 #include "CoreMinimal.h"
-#include "Engine/World.h"
 #include "AssetDefinition.h"
 
 #include "AssetDefinitionDefault.generated.h"
 
-
 UCLASS(Abstract)
 class UNREALED_API UAssetDefinitionDefault : public UAssetDefinition
 {
