AssetDefinition - Converted Fonts/FontFaces, SlateBrush, SlateWidgetStyle, and Lightweight instance.

diff --git a/Engine/Plugins/Editor/EngineAssetDefinitions/Source/EngineAssetDefinitions.Build.cs b/Engine/Plugins/Editor/EngineAssetDefinitions/Source/EngineAssetDefinitions.Build.cs
--- a/Engine/Plugins/Editor/EngineAssetDefinitions/Source/EngineAssetDefinitions.Build.cs
+++ b/Engine/Plugins/Editor/EngineAssetDefinitions/Source/EngineAssetDefinitions.Build.cs
@@ -44,6 +44,7 @@ public class EngineAssetDefinitions : ModuleRules
 				"TextureEditor",
 				"MaterialEditor",
 				"PhysicsAssetEditor",
+				"FontEditor"
 			}
 		);
 
@@ -55,6 +56,7 @@ public class EngineAssetDefinitions : ModuleRules
 				"TextureEditor",
 				"MaterialEditor",
 				"PhysicsAssetEditor",
+				"FontEditor"
 			}
 		);
 	}
diff --git a/Engine/Plugins/Editor/EngineAssetDefinitions/Source/Private/AssetDefinition_Font.cpp b/Engine/Plugins/Editor/EngineAssetDefinitions/Source/Private/AssetDefinition_Font.cpp
new file mode 100644
--- /dev/null
+++ b/Engine/Plugins/Editor/EngineAssetDefinitions/Source/Private/AssetDefinition_Font.cpp
@@ -0,0 +1,75 @@
+// Copyright Epic Games, Inc. All Rights Reserved.
+
+#include "AssetDefinition_Font.h"
+
+#include "ContentBrowserMenuContexts.h"
+#include "ToolMenus.h"
+#include "Styling/AppStyle.h"
+#include "EditorReimportHandler.h"
+#include "FontEditorModule.h"
+
+#define LOCTEXT_NAMESPACE "UAssetDefinition_Font"
+
+EAssetCommandResult UAssetDefinition_Font::OpenAssets(const FAssetOpenArgs& OpenArgs) const
+{
+	IFontEditorModule* FontEditorModule = &FModuleManager::LoadModuleChecked<IFontEditorModule>("FontEditor");
+	
+	for (UFont* Font : OpenArgs.LoadObjects<UFont>())
+	{
+		FontEditorModule->CreateFontEditor(OpenArgs.GetToolkitMode(), OpenArgs.ToolkitHost, Font);
+	}
+	
+	return EAssetCommandResult::Handled;
+}
+
+// Menu Extensions
+//--------------------------------------------------------------------
+
+namespace MenuExtension_Font
+{
+	void ExecuteReimport(const FToolMenuContext& InContext)
+	{
+		const UContentBrowserAssetContextMenuContext* Context = UContentBrowserAssetContextMenuContext::FindContextWithAssets(InContext);
+		for (UFont* Font : Context->LoadSelectedObjects<UFont>())
+		{
+			// Skip fonts that aren't using an offline cache, as they can't be reimported
+			if (Font->FontCacheType == EFontCacheType::Offline)
+			{
+				// Fonts fail to reimport if they ask for a new file if missing
+				FReimportManager::Instance()->Reimport(Font, /*bAskForNewFileIfMissing=*/false);
+			}
+		}
+	}
+	
+	static FDelayedAutoRegisterHelper DelayedAutoRegister(EDelayedRegisterRunPhase::EndOfEngineInit, []{ 
+		UToolMenus::RegisterStartupCallback(FSimpleMulticastDelegate::FDelegate::CreateLambda([]()
+		{
+			FToolMenuOwnerScoped OwnerScoped(UE_MODULE_NAME);
+			UToolMenu* Menu = UE::ContentBrowser::ExtendToolMenu_AssetContextMenu(UFont::StaticClass());
+		
+			FToolMenuSection& Section = Menu->FindOrAddSection("GetAssetActions");
+			Section.AddDynamicEntry("GetAssetActions_UStaticMesh", FNewToolMenuSectionDelegate::CreateLambda([](FToolMenuSection& InSection)
+			{
+				if (const UContentBrowserAssetContextMenuContext* Context = UContentBrowserAssetContextMenuContext::FindContextWithAssets(InSection))
+				{
+					// TODO NDarnell Having reimporting as a right click option for something like fonts is probably a
+					// bit much, I think for something more complex that isn't always re-importable maybe they shouldn't
+					// be right click options.  Though...I mean shouldn't even online fonts be reimportable?  They come
+					// from OTF and TTF files...those things *can* be reloaded.
+					// Also why do we even bother with this when we have IsImportable Asset()?
+					// TODO Investigate, see if anyone remembers.
+					
+					const TAttribute<FText> Label = LOCTEXT("ReimportFontLabel", "Reimport");
+					const TAttribute<FText> ToolTip = LOCTEXT("ReimportFontTooltip", "Reimport the selected font(s).");
+					const FSlateIcon Icon = FSlateIcon(FAppStyle::GetAppStyleSetName(), "ContentBrowser.AssetActions.ReimportAsset");
+					
+					FToolUIAction UIAction;
+					UIAction.ExecuteAction = FToolMenuExecuteAction::CreateStatic(&ExecuteReimport);
+					InSection.AddMenuEntry("ReimportFont", Label, ToolTip, Icon, UIAction);
+				}
+			}));
+		}));
+	});
+}
+
+#undef LOCTEXT_NAMESPACE
diff --git a/Engine/Plugins/Editor/EngineAssetDefinitions/Source/Private/AssetDefinition_Font.h b/Engine/Plugins/Editor/EngineAssetDefinitions/Source/Private/AssetDefinition_Font.h
new file mode 100644
--- /dev/null
+++ b/Engine/Plugins/Editor/EngineAssetDefinitions/Source/Private/AssetDefinition_Font.h
@@ -0,0 +1,36 @@
+// Copyright Epic Games, Inc. All Rights Reserved.
+
+#pragma once
+
+#include "CoreMinimal.h"
+#include "Engine/Font.h"
+#include "AssetDefinitionDefault.h"
+
+#include "AssetDefinition_Font.generated.h"
+
+UCLASS()
+class UAssetDefinition_Font : public UAssetDefinitionDefault
+{
+	GENERATED_BODY()
+
+public:
+	// UAssetDefinition Begin
+	virtual FText GetAssetDisplayName() const override { return NSLOCTEXT("AssetTypeActions", "AssetTypeActions_Font", "Font"); }
+	virtual FLinearColor GetAssetColor() const override { return FLinearColor(FColor(128,128,64)); }
+	virtual TSoftClassPtr<UObject> GetAssetClass() const override { return UFont::StaticClass(); }
+	virtual TConstArrayView<FAssetCategoryPath> GetAssetCategories() const override
+	{
+		static const auto Categories = { EAssetCategoryPaths::UI };
+		return Categories;
+	}
+	
+	virtual EAssetCommandResult OpenAssets(const FAssetOpenArgs& OpenArgs) const override;
+	// UAssetDefinition End
+
+private:
+	/** Can we execute a reimport for the selected objects? */
+	bool CanExecuteReimport(const TArray<TWeakObjectPtr<UFont>> Objects) const;
+
+	/** Handler for when Reimport is selected */
+	void ExecuteReimport(const TArray<TWeakObjectPtr<UFont>> Objects) const;
+};
diff --git a/Engine/Plugins/Editor/EngineAssetDefinitions/Source/Private/AssetDefinition_FontFace.cpp b/Engine/Plugins/Editor/EngineAssetDefinitions/Source/Private/AssetDefinition_FontFace.cpp
new file mode 100644
--- /dev/null
+++ b/Engine/Plugins/Editor/EngineAssetDefinitions/Source/Private/AssetDefinition_FontFace.cpp
@@ -0,0 +1,79 @@
+// Copyright Epic Games, Inc. All Rights Reserved.
+
+#include "AssetDefinition_FontFace.h"
+
+#include "ContentBrowserMenuContexts.h"
+#include "FontEditorModule.h"
+#include "ToolMenus.h"
+#include "Styling/AppStyle.h"
+#include "EditorReimportHandler.h"
+#include "Toolkits/SimpleAssetEditor.h"
+
+#define LOCTEXT_NAMESPACE "AssetDefinition_FontFace"
+
+EAssetCommandResult UAssetDefinition_FontFace::OpenAssets(const FAssetOpenArgs& OpenArgs) const
+{
+	// Load the FontEditor module to ensure that FFontFaceDetailsCustomization is registered
+	IFontEditorModule* FontEditorModule = &FModuleManager::LoadModuleChecked<IFontEditorModule>("FontEditor");
+	
+	// Open each object in turn, as the default editor would try and collapse a multi-edit together 
+	// into a single editor instance which doesn't really work for font face assets
+	for (UFontFace* FontFace : OpenArgs.LoadObjects<UFontFace>())
+	{
+		FSimpleAssetEditor::CreateEditor(OpenArgs.GetToolkitMode(), OpenArgs.ToolkitHost, FontFace);
+	}
+
+	return EAssetCommandResult::Handled;
+}
+
+// Menu Extensions
+//--------------------------------------------------------------------
+
+namespace MenuExtension_FontFace
+{
+	void ExecuteReimport(const FToolMenuContext& InContext)
+	{
+		const UContentBrowserAssetContextMenuContext* Context = UContentBrowserAssetContextMenuContext::FindContextWithAssets(InContext);
+		for (UFontFace* Font : Context->LoadSelectedObjects<UFontFace>())
+		{
+			// Skip fonts that don't have a source filename, as they can't be reimported
+			if (!Font->SourceFilename.IsEmpty())
+			{
+				// Fonts fail to reimport if they ask for a new file if missing
+				FReimportManager::Instance()->Reimport(Font, /*bAskForNewFileIfMissing=*/false);
+			}
+		}
+	}
+	
+	static FDelayedAutoRegisterHelper DelayedAutoRegister(EDelayedRegisterRunPhase::EndOfEngineInit, []{ 
+		UToolMenus::RegisterStartupCallback(FSimpleMulticastDelegate::FDelegate::CreateLambda([]()
+		{
+			FToolMenuOwnerScoped OwnerScoped(UE_MODULE_NAME);
+			UToolMenu* Menu = UE::ContentBrowser::ExtendToolMenu_AssetContextMenu(UFontFace::StaticClass());
+		
+			FToolMenuSection& Section = Menu->FindOrAddSection("GetAssetActions");
+			Section.AddDynamicEntry("GetAssetActions_UStaticMesh", FNewToolMenuSectionDelegate::CreateLambda([](FToolMenuSection& InSection)
+			{
+				if (const UContentBrowserAssetContextMenuContext* Context = UContentBrowserAssetContextMenuContext::FindContextWithAssets(InSection))
+				{
+					// TODO NDarnell Having reimporting as a right click option for something like fonts is probably a
+					// bit much, I think for something more complex that isn't always re-importable maybe they shouldn't
+					// be right click options.  Though...I mean shouldn't even online fonts be reimportable?  They come
+					// from OTF and TTF files...those things *can* be reloaded.
+					// Also why do we even bother with this when we have IsImportable Asset()?
+					// TODO Investigate, see if anyone remembers.
+					
+					const TAttribute<FText> Label = LOCTEXT("ReimportFontLabel", "Reimport");
+					const TAttribute<FText> ToolTip = LOCTEXT("ReimportFontTooltip", "Reimport the selected font(s).");
+					const FSlateIcon Icon = FSlateIcon(FAppStyle::GetAppStyleSetName(), "ContentBrowser.AssetActions.ReimportAsset");
+					
+					FToolUIAction UIAction;
+					UIAction.ExecuteAction = FToolMenuExecuteAction::CreateStatic(&ExecuteReimport);
+					InSection.AddMenuEntry("ReimportFont", Label, ToolTip, Icon, UIAction);
+				}
+			}));
+		}));
+	});
+}
+
+#undef LOCTEXT_NAMESPACE
diff --git a/Engine/Plugins/Editor/EngineAssetDefinitions/Source/Private/AssetDefinition_FontFace.h b/Engine/Plugins/Editor/EngineAssetDefinitions/Source/Private/AssetDefinition_FontFace.h
new file mode 100644
--- /dev/null
+++ b/Engine/Plugins/Editor/EngineAssetDefinitions/Source/Private/AssetDefinition_FontFace.h
@@ -0,0 +1,40 @@
+// Copyright Epic Games, Inc. All Rights Reserved.
+
+#pragma once
+
+#include "CoreMinimal.h"
+#include "Engine/FontFace.h"
+#include "AssetDefinitionDefault.h"
+
+#include "AssetDefinition_FontFace.generated.h"
+
+UCLASS()
+class UAssetDefinition_FontFace : public UAssetDefinitionDefault
+{
+	GENERATED_BODY()
+
+public:
+	// UAssetDefinition Begin
+	virtual FText GetAssetDisplayName() const override { return NSLOCTEXT("AssetTypeActions", "AssetTypeActions_FontFace", "Font Face"); }
+	virtual FLinearColor GetAssetColor() const override { return FLinearColor(FColor(184,184,112)); }
+	virtual TSoftClassPtr<UObject> GetAssetClass() const override { return UFontFace::StaticClass(); }
+	virtual TConstArrayView<FAssetCategoryPath> GetAssetCategories() const override
+	{
+		static const auto Categories = { EAssetCategoryPaths::UI };
+		return Categories;
+	}
+	virtual FAssetSupportResponse CanLocalize(const FAssetData& InAsset) const
+	{
+		return FAssetSupportResponse::NotSupported();
+	}
+	
+	virtual EAssetCommandResult OpenAssets(const FAssetOpenArgs& OpenArgs) const override;
+	// UAssetDefinition End
+	
+private:
+	/** Can we execute a reimport for the selected objects? */
+	bool CanExecuteReimport(const TArray<TWeakObjectPtr<UFontFace>> Objects) const;
+
+	/** Handler for when Reimport is selected */
+	void ExecuteReimport(const TArray<TWeakObjectPtr<UFontFace>> Objects) const;
+};
diff --git a/Engine/Plugins/Editor/EngineAssetDefinitions/Source/Private/AssetDefinition_LightWeightInstance.h b/Engine/Plugins/Editor/EngineAssetDefinitions/Source/Private/AssetDefinition_LightWeightInstance.h
new file mode 100644
--- /dev/null
+++ b/Engine/Plugins/Editor/EngineAssetDefinitions/Source/Private/AssetDefinition_LightWeightInstance.h
@@ -0,0 +1,27 @@
+// Copyright Epic Games, Inc. All Rights Reserved.
+
+#pragma once
+
+#include "CoreMinimal.h"
+#include "GameFramework/LightWeightInstanceManager.h"
+#include "AssetDefinitionDefault.h"
+
+#include "AssetDefinition_LightWeightInstance.generated.h"
+
+UCLASS()
+class UAssetDefinition_LightWeightInstance : public UAssetDefinitionDefault
+{
+	GENERATED_BODY()
+
+public:
+	// UAssetDefinition Begin
+	virtual FText GetAssetDisplayName() const override { return NSLOCTEXT("AssetTypeActions", "AssetTypeActions_LightWeightInstance", "Light Weight Instance"); }
+	virtual FLinearColor GetAssetColor() const override { return FLinearColor(FColor(255, 105, 180)); }
+	virtual TSoftClassPtr<UObject> GetAssetClass() const override { return ALightWeightInstanceManager::StaticClass(); }
+	virtual TConstArrayView<FAssetCategoryPath> GetAssetCategories() const override
+	{
+		static const auto Categories = { EAssetCategoryPaths::Gameplay };
+		return Categories;
+	}
+	// UAssetDefinition End
+};
diff --git a/Engine/Plugins/Editor/EngineAssetDefinitions/Source/Private/AssetDefinition_SlateBrush.h b/Engine/Plugins/Editor/EngineAssetDefinitions/Source/Private/AssetDefinition_SlateBrush.h
new file mode 100644
--- /dev/null
+++ b/Engine/Plugins/Editor/EngineAssetDefinitions/Source/Private/AssetDefinition_SlateBrush.h
@@ -0,0 +1,26 @@
+// Copyright Epic Games, Inc. All Rights Reserved.
+
+#pragma once
+
+#include "CoreMinimal.h"
+#include "Slate/SlateBrushAsset.h"
+#include "AssetDefinitionDefault.h"
+
+#include "AssetDefinition_SlateBrush.generated.h"
+
+UCLASS()
+class UAssetDefinition_SlateBrush : public UAssetDefinitionDefault
+{
+	GENERATED_BODY()
+
+public:
+	// UAssetDefinition Begin
+	virtual FText GetAssetDisplayName() const override { return NSLOCTEXT("AssetTypeActions", "AssetTypeActions_SlateBrush", "Slate Brush"); }
+	virtual FLinearColor GetAssetColor() const override { return FLinearColor(FColor(105, 165, 60)); }
+	virtual TSoftClassPtr<UObject> GetAssetClass() const override { return USlateBrushAsset::StaticClass(); }
+	virtual TConstArrayView<FAssetCategoryPath> GetAssetCategories() const override
+	{
+		static const auto Categories = { EAssetCategoryPaths::UI };
+		return Categories;
+	}
+};
diff --git a/Engine/Plugins/Editor/EngineAssetDefinitions/Source/Private/AssetDefinition_SlateWidgetStyle.cpp b/Engine/Plugins/Editor/EngineAssetDefinitions/Source/Private/AssetDefinition_SlateWidgetStyle.cpp
new file mode 100644
--- /dev/null
+++ b/Engine/Plugins/Editor/EngineAssetDefinitions/Source/Private/AssetDefinition_SlateWidgetStyle.cpp
@@ -0,0 +1,34 @@
+// Copyright Epic Games, Inc. All Rights Reserved.
+
+#include "AssetDefinition_SlateWidgetStyle.h"
+#include "ToolMenus.h"
+#include "Toolkits/SimpleAssetEditor.h"
+
+#define LOCTEXT_NAMESPACE "UAssetDefinition_SlateWidgetStyle"
+
+EAssetCommandResult UAssetDefinition_SlateWidgetStyle::OpenAssets(const FAssetOpenArgs& OpenArgs) const
+{
+	struct Local
+	{
+		static TArray<UObject*> GetSubObjects(const TArray<UObject*>& InObjects)
+		{
+			TArray<UObject*> SubObjects;
+			for(UObject* Object : InObjects)
+			{
+				auto Style = Cast<USlateWidgetStyleAsset>(Object);
+				if(Style && Style->CustomStyle)
+				{
+					SubObjects.Add(Style->CustomStyle);
+				}
+			}
+			return SubObjects;
+		}
+	};
+
+	const TArray<UObject*> Objects(OpenArgs.LoadObjects<USlateWidgetStyleAsset>());
+	FSimpleAssetEditor::CreateEditor(OpenArgs.GetToolkitMode(), OpenArgs.ToolkitHost, Objects, FSimpleAssetEditor::FGetDetailsViewObjects::CreateStatic(&Local::GetSubObjects));
+
+	return EAssetCommandResult::Handled;
+}
+
+#undef LOCTEXT_NAMESPACE
diff --git a/Engine/Plugins/Editor/EngineAssetDefinitions/Source/Private/AssetDefinition_SlateWidgetStyle.h b/Engine/Plugins/Editor/EngineAssetDefinitions/Source/Private/AssetDefinition_SlateWidgetStyle.h
new file mode 100644
--- /dev/null
+++ b/Engine/Plugins/Editor/EngineAssetDefinitions/Source/Private/AssetDefinition_SlateWidgetStyle.h
@@ -0,0 +1,29 @@
+// Copyright Epic Games, Inc. All Rights Reserved.
+
+#pragma once
+
+#include "CoreMinimal.h"
+#include "Styling/SlateWidgetStyleAsset.h"
+#include "AssetDefinitionDefault.h"
+
+#include "AssetDefinition_SlateWidgetStyle.generated.h"
+
+UCLASS()
+class UAssetDefinition_SlateWidgetStyle : public UAssetDefinitionDefault
+{
+	GENERATED_BODY()
+
+public:
+	// UAssetDefinition Begin
+	virtual FText GetAssetDisplayName() const override { return NSLOCTEXT("AssetTypeActions", "AssetTypeActions_SlateStyle", "Slate Widget Style"); }
+	virtual FLinearColor GetAssetColor() const override { return FLinearColor(FColor(62, 140, 35)); }
+	virtual TSoftClassPtr<UObject> GetAssetClass() const override { return USlateWidgetStyleAsset::StaticClass(); }
+	virtual TConstArrayView<FAssetCategoryPath> GetAssetCategories() const override
+	{
+		static const auto Categories = { EAssetCategoryPaths::UI };
+		return Categories;
+	}
+	
+	virtual EAssetCommandResult OpenAssets(const FAssetOpenArgs& OpenArgs) const override;
+	// UAssetDefinition End
+};
diff --git a/Engine/Source/Developer/AssetTools/Private/AssetTools.cpp b/Engine/Source/Developer/AssetTools/Private/AssetTools.cpp
--- a/Engine/Source/Developer/AssetTools/Private/AssetTools.cpp
+++ b/Engine/Source/Developer/AssetTools/Private/AssetTools.cpp
@@ -58,17 +58,12 @@
 #include "AssetTypeActions/AssetTypeActions_DataTable.h"
 #include "AssetTypeActions/AssetTypeActions_CompositeDataTable.h"
 #include "AssetTypeActions/AssetTypeActions_Class.h"
-#include "AssetTypeActions/AssetTypeActions_Font.h"
-#include "AssetTypeActions/AssetTypeActions_FontFace.h"
 #include "AssetTypeActions/AssetTypeActions_ForceFeedbackEffect.h"
 #include "AssetTypeActions/AssetTypeActions_ActorFoliageSettings.h"
-#include "AssetTypeActions/AssetTypeActions_LightWeightInstance.h"
 #include "AssetTypeActions/AssetTypeActions_MirrorDataTable.h"
 #include "AssetTypeActions/AssetTypeActions_ParticleSystem.h"
 #include "AssetTypeActions/AssetTypeActions_PhysicalMaterialMask.h"
 #include "AssetTypeActions/AssetTypeActions_Skeleton.h"
-#include "AssetTypeActions/AssetTypeActions_SlateBrush.h"
-#include "AssetTypeActions/AssetTypeActions_SlateWidgetStyle.h"
 #include "WorldPartition/WorldPartition.h"
 #include "SDiscoveringAssetsDialog.h"
 #include "AssetFixUpRedirectors.h"
@@ -1256,18 +1251,13 @@ UAssetToolsImpl::UAssetToolsImpl(const FObjectInitializer& ObjectInitializer)
 	RegisterAssetTypeActions(MakeShareable(new FAssetTypeActions_DataTable));
 	RegisterAssetTypeActions(MakeShareable(new FAssetTypeActions_CompositeDataTable));
 	RegisterAssetTypeActions(MakeShareable(new FAssetTypeActions_Class));
-	RegisterAssetTypeActions(MakeShareable(new FAssetTypeActions_Font));
-	RegisterAssetTypeActions(MakeShareable(new FAssetTypeActions_FontFace));
 	RegisterAssetTypeActions(MakeShareable(new FAssetTypeActions_ForceFeedbackEffect(InputCategoryBit)));
 	RegisterAssetTypeActions(MakeShareable(new FAssetTypeActions_ActorFoliageSettings(FoliageCategoryBit)));
-	RegisterAssetTypeActions(MakeShareable(new FAssetTypeActions_LightWeightInstance));
 	RegisterAssetTypeActions(MakeShareable(new FAssetTypeActions_MirrorDataTable));
 	RegisterAssetTypeActions(MakeShareable(new FAssetTypeActions_ParticleSystem));
 	RegisterAssetTypeActions(MakeShareable(new FAssetTypeActions_PhysicalMaterialMask));
 	RegisterAssetTypeActions(MakeShareable(new FAssetTypeActions_SkeletalMesh));
 	RegisterAssetTypeActions(MakeShareable(new FAssetTypeActions_Skeleton));
-	RegisterAssetTypeActions(MakeShareable(new FAssetTypeActions_SlateBrush));
-	RegisterAssetTypeActions(MakeShareable(new FAssetTypeActions_SlateWidgetStyle));
 
 	// Note: Please don't add any more actions here!  They belong in an editor-only module that is more tightly
 	// coupled to your new system, and you should not create a dependency on your new system from AssetTools.
diff --git a/Engine/Source/Developer/AssetTools/Private/AssetTypeActions/AssetTypeActions_Font.cpp b/Engine/Source/Developer/AssetTools/Private/AssetTypeActions/AssetTypeActions_Font.cpp
deleted file mode 100644
diff --git a/Engine/Source/Developer/AssetTools/Private/AssetTypeActions/AssetTypeActions_Font.h b/Engine/Source/Developer/AssetTools/Private/AssetTypeActions/AssetTypeActions_Font.h
deleted file mode 100644
diff --git a/Engine/Source/Developer/AssetTools/Private/AssetTypeActions/AssetTypeActions_FontFace.cpp b/Engine/Source/Developer/AssetTools/Private/AssetTypeActions/AssetTypeActions_FontFace.cpp
deleted file mode 100644
diff --git a/Engine/Source/Developer/AssetTools/Private/AssetTypeActions/AssetTypeActions_FontFace.h b/Engine/Source/Developer/AssetTools/Private/AssetTypeActions/AssetTypeActions_FontFace.h
deleted file mode 100644
diff --git a/Engine/Source/Developer/AssetTools/Private/AssetTypeActions/AssetTypeActions_LightWeightInstance.cpp b/Engine/Source/Developer/AssetTools/Private/AssetTypeActions/AssetTypeActions_LightWeightInstance.cpp
deleted file mode 100644
diff --git a/Engine/Source/Developer/AssetTools/Private/AssetTypeActions/AssetTypeActions_LightWeightInstance.h b/Engine/Source/Developer/AssetTools/Private/AssetTypeActions/AssetTypeActions_LightWeightInstance.h
deleted file mode 100644
diff --git a/Engine/Source/Developer/AssetTools/Private/AssetTypeActions/AssetTypeActions_SlateBrush.cpp b/Engine/Source/Developer/AssetTools/Private/AssetTypeActions/AssetTypeActions_SlateBrush.cpp
deleted file mode 100644
diff --git a/Engine/Source/Developer/AssetTools/Private/AssetTypeActions/AssetTypeActions_SlateBrush.h b/Engine/Source/Developer/AssetTools/Private/AssetTypeActions/AssetTypeActions_SlateBrush.h
deleted file mode 100644
diff --git a/Engine/Source/Developer/AssetTools/Private/AssetTypeActions/AssetTypeActions_SlateWidgetStyle.cpp b/Engine/Source/Developer/AssetTools/Private/AssetTypeActions/AssetTypeActions_SlateWidgetStyle.cpp
deleted file mode 100644
diff --git a/Engine/Source/Developer/AssetTools/Private/AssetTypeActions/AssetTypeActions_SlateWidgetStyle.h b/Engine/Source/Developer/AssetTools/Private/AssetTypeActions/AssetTypeActions_SlateWidgetStyle.h
deleted file mode 100644
