AssetDefinition - StaticMesh's Nanite menu restored to its former glory.  The checkbox now works with multiple selected meshes.

diff --git a/Engine/Plugins/Editor/EngineAssetDefinitions/Source/Private/AssetDefinition_StaticMesh.cpp b/Engine/Plugins/Editor/EngineAssetDefinitions/Source/Private/AssetDefinition_StaticMesh.cpp
--- a/Engine/Plugins/Editor/EngineAssetDefinitions/Source/Private/AssetDefinition_StaticMesh.cpp
+++ b/Engine/Plugins/Editor/EngineAssetDefinitions/Source/Private/AssetDefinition_StaticMesh.cpp
@@ -17,6 +17,8 @@
 #include "ToolMenus.h"
 #include "Widgets/Notifications/SNotificationList.h"
 #include "ContentBrowserMenuContexts.h"
+#include "Algo/AllOf.h"
+#include "Algo/Compare.h"
 
 #define LOCTEXT_NAMESPACE "UAssetDefinition_StaticMesh"
 
@@ -158,6 +160,46 @@ namespace MenuExtension_StaticMesh
 	
 		FToolMenuSection& Section = Menu->FindOrAddSection("NaniteActions");
 
+		{
+			const TAttribute<FText> Label = LOCTEXT("StaticMesh_NaniteToggle", "Nanite");
+			const TAttribute<FText> ToolTip = LOCTEXT("StaticMesh_NaniteToggleTooltip", "Toggle Nanite support on the selected meshes.");
+			const FSlateIcon Icon = FSlateIcon();
+
+			FToolUIAction UIAction;
+			UIAction.ExecuteAction = FToolMenuExecuteAction::CreateLambda([](const FToolMenuContext& InContext)
+			{
+				if (const UContentBrowserAssetContextMenuContext* Context = UContentBrowserAssetContextMenuContext::FindContextWithAssets(InContext))
+				{
+					const bool bContainsTrue = Context->SelectedAssets.ContainsByPredicate([](const FAssetData& InAsset) { return IsNaniteEnabled(InAsset); });
+					const bool bContainsFalse = Context->SelectedAssets.ContainsByPredicate([](const FAssetData& InAsset) { return !IsNaniteEnabled(InAsset); });
+
+					if ((bContainsTrue && bContainsFalse) || bContainsTrue)
+					{
+						ExecuteNaniteDisable(InContext);
+					}
+					else
+					{
+						ExecuteNaniteEnable(InContext);
+					}
+				}
+			});
+			UIAction.GetActionCheckState = FToolMenuGetActionCheckState::CreateLambda([](const FToolMenuContext& InContext)
+			{
+				if (const UContentBrowserAssetContextMenuContext* Context = UContentBrowserAssetContextMenuContext::FindContextWithAssets(InContext))
+				{
+					const bool bContainsTrue = Context->SelectedAssets.ContainsByPredicate([](const FAssetData& InAsset) { return IsNaniteEnabled(InAsset); });
+					const bool bContainsFalse = Context->SelectedAssets.ContainsByPredicate([](const FAssetData& InAsset) { return !IsNaniteEnabled(InAsset); });
+					return bContainsTrue && bContainsFalse ? ECheckBoxState::Undetermined : (bContainsTrue ? ECheckBoxState::Checked : ECheckBoxState::Unchecked);
+				}
+
+				return ECheckBoxState::Undetermined;
+			});
+			
+			Section.AddMenuEntry("StaticMesh_NaniteToggle", Label, ToolTip, Icon, UIAction, EUserInterfaceActionType::ToggleButton);
+		}
+
+		Section.AddSeparator("Nanite_EnableDisableOptions");
+
 		const int32 NaniteMeshes = Algo::Accumulate(Context->SelectedAssets, 0, [](int32 Value, const FAssetData& StaticMeshAsset)
 			{
 				bool bNaniteEnabled = false;
@@ -167,7 +209,7 @@ namespace MenuExtension_StaticMesh
 		{
 			const int32 RegularMeshes = Context->SelectedAssets.Num() - NaniteMeshes;
 			
-			const TAttribute<FText> Label = FText::Format(LOCTEXT("StaticMesh_NaniteEnableAll", "Enable Nanite (% Meshes)"), RegularMeshes);
+			const TAttribute<FText> Label = FText::Format(LOCTEXT("StaticMesh_NaniteEnableAll", "Enable Nanite ({0} Meshes)"), RegularMeshes);
 			const TAttribute<FText> ToolTip = LOCTEXT("StaticMesh_NaniteEnableAllTooltip", "Enables support for Nanite on the selected meshes.");
 			const FSlateIcon Icon = FSlateIcon();
 			
@@ -178,7 +220,7 @@ namespace MenuExtension_StaticMesh
 		}
 		
 		{
-			const TAttribute<FText> Label = FText::Format(LOCTEXT("StaticMesh_NaniteDisableAll", "Disable Nanite (% Meshes)"), NaniteMeshes);
+			const TAttribute<FText> Label = FText::Format(LOCTEXT("StaticMesh_NaniteDisableAll", "Disable Nanite ({0} Meshes)"), NaniteMeshes);
 			const TAttribute<FText> ToolTip = LOCTEXT("StaticMesh_NaniteDisableAllTooltip", "Disables support for Nanite on the selected meshes.");
 			const FSlateIcon Icon = FSlateIcon();
 			
@@ -496,7 +538,7 @@ namespace MenuExtension_StaticMesh
 					}
 					{
 						InSection.AddSubMenu(
-							"StaticMesh_LODMenu",
+							"StaticMesh_NaniteMenu",
 							LOCTEXT("StaticMesh_NaniteMenu", "Nanite"),
 							LOCTEXT("StaticMesh_NaniteTooltip", "Nanite Options and Tools"),
 							FNewToolMenuDelegate::CreateStatic(&GetNaniteMenu),
