AssetDefinition - "UV Editor" plugin's option no longer hides if it won't work (determining it requires loading the assets), so it's just going to determine that when you fire it and fire a notification if it wont work.
AssetDefinition - Skeleton
AssetDefinition - SkeletonMesh
AssetDefinition - StringTable

diff --git a/Engine/Plugins/Editor/EngineAssetDefinitions/Source/Private/Animation/AssetDefinition_SkeletalMesh.cpp b/Engine/Plugins/Editor/EngineAssetDefinitions/Source/Private/Animation/AssetDefinition_SkeletalMesh.cpp
new file mode 100644
--- /dev/null
+++ b/Engine/Plugins/Editor/EngineAssetDefinitions/Source/Private/Animation/AssetDefinition_SkeletalMesh.cpp
@@ -0,0 +1,457 @@
+// Copyright Epic Games, Inc. All Rights Reserved.
+
+#include "Animation/AssetDefinition_SkeletalMesh.h"
+
+#include "AnimationEditorUtils.h"
+#include "AssetNotifications.h"
+#include "AssetToolsModule.h"
+#include "ContentBrowserMenuContexts.h"
+#include "ContentBrowserModule.h"
+#include "EditorReimportHandler.h"
+#include "FbxMeshUtils.h"
+#include "ISkeletalMeshEditorModule.h"
+#include "Animation/SkeletalMeshUtilities.h"
+#include "AssetRegistry/AssetRegistryModule.h"
+#include "Factories/SkeletonFactory.h"
+#include "ThumbnailRendering/SceneThumbnailInfo.h"
+#include "Misc/MessageDialog.h"
+#include "Animation/Skeleton.h"
+#include "Factories/FbxSkeletalMeshImportData.h"
+#include "Factories/PhysicsAssetFactory.h"
+#include "PhysicsEngine/PhysicsAsset.h"
+#include "Widgets/Layout/SBorder.h"
+#include "Widgets/Images/SImage.h"
+#include "ToolMenu.h"
+#include "ToolMenus.h"
+
+#define LOCTEXT_NAMESPACE "AssetTypeActions"
+
+UAssetDefinition_SkeletalMesh::UAssetDefinition_SkeletalMesh()
+{
+	//No need to remove since the asset registry module clear this multicast delegate when terminating
+	FAssetRegistryModule& AssetRegistryModule = FModuleManager::GetModuleChecked<FAssetRegistryModule>(TEXT("AssetRegistry"));
+	AssetRegistryModule.Get().OnAssetRemoved().AddUObject(this, &UAssetDefinition_SkeletalMesh::OnAssetRemoved);
+}
+
+EAssetCommandResult UAssetDefinition_SkeletalMesh::OpenAssets(const FAssetOpenArgs& OpenArgs) const
+{
+	for (USkeletalMesh* Mesh : OpenArgs.LoadObjects<USkeletalMesh>())
+	{
+		if (Mesh->GetSkeleton() == nullptr)
+		{
+			if ( FMessageDialog::Open(EAppMsgType::YesNo, LOCTEXT("MissingSkeleton", "This mesh currently has no valid Skeleton. Would you like to create a new Skeleton?")) == EAppReturnType::Yes )
+			{
+				FAssetToolsModule& AssetToolsModule = FModuleManager::Get().LoadModuleChecked<FAssetToolsModule>("AssetTools");
+				
+				const FString DefaultSuffix = TEXT("_Skeleton");
+
+				// Determine an appropriate name
+				FString Name;
+				FString PackageName;
+				AssetToolsModule.Get().CreateUniqueAssetName(Mesh->GetOutermost()->GetName(), DefaultSuffix, PackageName, Name);
+
+				USkeletonFactory* Factory = NewObject<USkeletonFactory>();
+				Factory->TargetSkeletalMesh = Mesh;
+				
+				AssetToolsModule.Get().CreateAsset(Name, FPackageName::GetLongPackagePath(PackageName), USkeleton::StaticClass(), Factory);
+			}
+			else
+			{
+				UE::SkeletalMeshUtilities::AssignSkeletonToMesh(Mesh);
+			}
+
+			if( Mesh->GetSkeleton() == nullptr )
+			{
+				// error message
+				FMessageDialog::Open( EAppMsgType::Ok, LOCTEXT("CreateSkeletonOrAssign", "You need to create a Skeleton or assign one in order to open this in Persona."));
+			}
+		}
+
+		if ( Mesh->GetSkeleton() != nullptr )
+		{
+			ISkeletalMeshEditorModule& SkeletalMeshEditorModule = FModuleManager::LoadModuleChecked<ISkeletalMeshEditorModule>("SkeletalMeshEditor");
+			SkeletalMeshEditorModule.CreateSkeletalMeshEditor(OpenArgs.GetToolkitMode(), OpenArgs.ToolkitHost, Mesh);
+		}
+	}
+	
+	return EAssetCommandResult::Handled;
+}
+
+UThumbnailInfo* UAssetDefinition_SkeletalMesh::LoadThumbnailInfo(const FAssetData& InAssetData) const
+{
+	USkeletalMesh* SkeletalMesh = CastChecked<USkeletalMesh>(InAssetData.GetAsset());
+	UThumbnailInfo* ThumbnailInfo = SkeletalMesh->GetThumbnailInfo();
+	if ( ThumbnailInfo == nullptr )
+	{
+		ThumbnailInfo = NewObject<USceneThumbnailInfo>(SkeletalMesh, NAME_None, RF_Transactional);
+		SkeletalMesh->SetThumbnailInfo(ThumbnailInfo);
+	}
+
+	return ThumbnailInfo;
+}
+
+EVisibility UAssetDefinition_SkeletalMesh::GetThumbnailSkinningOverlayVisibility(const FAssetData AssetData) const
+{
+	//If the asset was delete it will be remove from the list, in that case do not use the
+	//asset since it can be invalid if the GC has collect the object point by the AssetData.
+	if (!ThumbnailSkinningOverlayAssetNames.Contains(AssetData.GetFullName()))
+	{
+		return EVisibility::Collapsed;
+	}
+
+	//Prevent loading assets when we display the thumbnail use the asset registry tags
+
+	//Legacy fbx code
+	FAssetDataTagMapSharedView::FFindTagResult Result = AssetData.TagsAndValues.FindTag(GET_MEMBER_NAME_CHECKED(UFbxSkeletalMeshImportData, LastImportContentType));
+	if (Result.IsSet() && Result.GetValue() == TEXT("FBXICT_Geometry"))
+	{
+		//Show the icon
+		return EVisibility::HitTestInvisible;
+	}
+
+#if WITH_EDITORONLY_DATA
+	//Generic Interchange code
+	Result = AssetData.TagsAndValues.FindTag(NSSkeletalMeshSourceFileLabels::GetSkeletalMeshLastImportContentTypeMetadataKey());
+	if (Result.IsSet() && Result.GetValue().Equals(NSSkeletalMeshSourceFileLabels::GeometryMetaDataValue()))
+	{
+		//Show the icon
+		return EVisibility::HitTestInvisible;
+	}
+#endif
+	
+	return EVisibility::Collapsed;
+}
+
+void UAssetDefinition_SkeletalMesh::OnAssetRemoved(const FAssetData& AssetData) const
+{
+	//Remove the object from the list before it get garbage collect
+	if (ThumbnailSkinningOverlayAssetNames.Contains(AssetData.GetFullName()))
+	{
+		ThumbnailSkinningOverlayAssetNames.Remove(AssetData.GetFullName());
+	}
+}
+
+TSharedPtr<SWidget> UAssetDefinition_SkeletalMesh::GetThumbnailOverlay(const FAssetData& AssetData) const
+{
+	const FSlateBrush* Icon = FAppStyle::GetBrush("ClassThumbnailOverlays.SkeletalMesh_NeedSkinning");
+
+	ThumbnailSkinningOverlayAssetNames.AddUnique(AssetData.GetFullName());
+
+	return SNew(SBorder)
+		.BorderImage(FAppStyle::GetNoBrush())
+		.Visibility_UObject(this, &UAssetDefinition_SkeletalMesh::GetThumbnailSkinningOverlayVisibility, AssetData)
+		.Padding(FMargin(0.0f, 0.0f, 3.0f, 3.0f))
+		.HAlign(HAlign_Right)
+		.VAlign(VAlign_Bottom)
+		[
+			SNew(SImage)
+			.ToolTipText(LOCTEXT("FAssetTypeActions_SkeletalMesh_NeedSkinning_ToolTip", "Asset geometry was imported, the skinning need to be validate"))
+			.Image(Icon)
+		];
+}
+
+// Menu Extensions
+//--------------------------------------------------------------------
+
+namespace MenuExtension_SkeletalMesh
+{
+	void ExecuteNewPhysicsAsset(const FToolMenuContext& InContext, bool bSetAssetToMesh)
+	{
+		TArray<UObject*> CreatedObjects;
+		
+		const UContentBrowserAssetContextMenuContext* Context = UContentBrowserAssetContextMenuContext::FindContextWithAssets(InContext);
+		for (USkeletalMesh* SkeletalMesh : Context->LoadSelectedObjects<USkeletalMesh>())
+		{
+			if (SkeletalMesh->GetOutermost()->bIsCookedForEditor)
+			{
+				FAssetNotifications::CannotEditCookedAsset(SkeletalMesh);
+				continue;
+			}
+			
+			if (UObject* PhysicsAsset = UPhysicsAssetFactory::CreatePhysicsAssetFromMesh(NAME_None, nullptr, SkeletalMesh, bSetAssetToMesh))
+			{
+				CreatedObjects.Add(PhysicsAsset);
+			}
+		}
+    
+		if(CreatedObjects.Num() > 0)
+		{
+			FContentBrowserModule& ContentBrowserModule = FModuleManager::LoadModuleChecked<FContentBrowserModule>("ContentBrowser");
+			ContentBrowserModule.Get().SyncBrowserToAssets(CreatedObjects);
+#if WITH_EDITOR
+			GEditor->GetEditorSubsystem<UAssetEditorSubsystem>()->OpenEditorForAssets(CreatedObjects);
+#endif
+		}
+	}
+
+	void GetPhysicsAssetMenu(UToolMenu* Menu)
+	{
+		const UContentBrowserAssetContextMenuContext* Context = UContentBrowserAssetContextMenuContext::FindContextWithAssets(*Menu);
+
+    	FToolMenuSection& Section = Menu->FindOrAddSection("PhysicsActions");
+		
+		{
+			const TAttribute<FText> Label = LOCTEXT("PhysAsset_Create", "Create");
+			const TAttribute<FText> ToolTip = LOCTEXT("PhysAsset_Create_ToolTip", "Create new physics assets without assigning it to the selected meshes");
+			const FSlateIcon Icon = FSlateIcon();
+
+			FToolUIAction UIAction;
+			UIAction.ExecuteAction = FToolMenuExecuteAction::CreateStatic(&ExecuteNewPhysicsAsset, false);
+			Section.AddMenuEntry("PhysAsset_Create", Label, ToolTip, Icon, UIAction);
+		}
+		
+		{
+			const TAttribute<FText> Label = LOCTEXT("PhysAsset_CreateAssign", "Create and Assign");
+			const TAttribute<FText> ToolTip = LOCTEXT("PhysAsset_CreateAssign_ToolTip", "Create new physics assets and assign it to each of the selected meshes");
+			const FSlateIcon Icon = FSlateIcon();
+
+			FToolUIAction UIAction;
+			UIAction.ExecuteAction = FToolMenuExecuteAction::CreateStatic(&ExecuteNewPhysicsAsset, true);
+			Section.AddMenuEntry("PhysAsset_CreateAssign", Label, ToolTip, Icon, UIAction);
+		}
+	}
+
+	bool OnAssetCreated(TArray<UObject*> NewAssets)
+	{
+		if (NewAssets.Num() > 1)
+		{
+			FAssetToolsModule& AssetToolsModule = FModuleManager::Get().LoadModuleChecked<FAssetToolsModule>("AssetTools");
+			AssetToolsModule.Get().SyncBrowserToAssets(NewAssets);
+		}
+		return true;
+	}
+	
+    void FillCreateMenu(UToolMenu* Menu)
+    {
+    	IAssetTools& AssetTools = FModuleManager::LoadModuleChecked<FAssetToolsModule>("AssetTools").Get();
+    	if (AssetTools.IsAssetClassSupported(UPhysicsAsset::StaticClass()))
+    	{
+    		FToolMenuSection& Section = Menu->AddSection("CreatePhysicsAsset", LOCTEXT("CreatePhysicsAssetMenuHeading", "Physics Asset"));
+    		Section.AddSubMenu(
+    			"NewPhysicsAssetMenu",
+    			LOCTEXT("SkeletalMesh_NewPhysicsAssetMenu", "Physics Asset"),
+    			LOCTEXT("SkeletalMesh_NewPhysicsAssetMenu_ToolTip", "Options for creating new physics assets from the selected meshes."),
+    			FNewToolMenuDelegate::CreateStatic(&GetPhysicsAssetMenu));
+    	}
+
+    	Menu->AddDynamicSection(NAME_None, FNewToolMenuDelegateLegacy::CreateLambda([](FMenuBuilder& MenuBuilder, UToolMenu* Menu)
+    	{
+    		const UContentBrowserAssetContextMenuContext* Context = UContentBrowserAssetContextMenuContext::FindContextWithAssets(*Menu);
+			const TArray<TSoftObjectPtr<UObject>> SkeletalMeshes = Context->GetSelectedAssetSoftObjects<UObject>();
+    		
+    		AnimationEditorUtils::FillCreateAssetMenu(MenuBuilder, SkeletalMeshes, FAnimAssetCreated::CreateStatic(&OnAssetCreated));
+    	}));
+    }
+    
+    void GetResolvedSourceFilePaths(const TArray<UObject*>& TypeAssets, TArray<FString>& OutSourceFilePaths)
+    {
+    	for (auto& Asset : TypeAssets)
+    	{
+    		const auto SkeletalMesh = CastChecked<USkeletalMesh>(Asset);
+    		SkeletalMesh->GetAssetImportData()->ExtractFilenames(OutSourceFilePaths);
+    	}
+    }
+    
+    void GetSourceFileLabels(const TArray<UObject*>& TypeAssets, TArray<FString>& OutSourceFileLabels)
+    {
+    	for (auto& Asset : TypeAssets)
+    	{
+    		const auto SkeletalMesh = CastChecked<USkeletalMesh>(Asset);
+    		TArray<FString> SourceFilePaths;
+    		SkeletalMesh->GetAssetImportData()->ExtractFilenames(SourceFilePaths);
+    		for (int32 SourceIndex = 0; SourceIndex < SourceFilePaths.Num(); ++SourceIndex)
+    		{
+    			FText SourceIndexLabel = USkeletalMesh::GetSourceFileLabelFromIndex(SourceIndex);
+    			OutSourceFileLabels.Add(SourceIndexLabel.ToString());
+    		}
+    	}
+    }
+
+	void ExecuteImportMeshLOD(const FToolMenuContext&, TSoftObjectPtr<USkeletalMesh> SkeletalMeshPtr, int32 LOD)
+	{
+		if (USkeletalMesh* SkeletalMesh = SkeletalMeshPtr.LoadSynchronous())
+		{
+			if (SkeletalMesh->GetOutermost()->bIsCookedForEditor)
+			{
+				FAssetNotifications::CannotEditCookedAsset(SkeletalMesh);
+				return;
+			}
+    
+			if (LOD == 0)
+			{
+				//re-import of the asset
+				TArray<UObject*> AssetArray;
+				AssetArray.Add(SkeletalMesh);
+				FReimportManager::Instance()->ValidateAllSourceFileAndReimport(AssetArray);
+			}
+			else
+			{
+				FbxMeshUtils::ImportMeshLODDialog(SkeletalMesh, LOD);
+			}
+		}
+	}
+
+	const static FName NAME_LODs("LODs");
+	static int32 GetNumberOfLODs(const FAssetData& SkeletalMeshAsset)
+	{
+		int32 NumberOfLODs = 0;
+		return SkeletalMeshAsset.GetTagValue<int32>(NAME_LODs, NumberOfLODs) ? NumberOfLODs : 0;
+	}
+    
+    void GetLODMenu(UToolMenu* Menu)
+    {
+		const UContentBrowserAssetContextMenuContext* Context = UContentBrowserAssetContextMenuContext::FindContextWithAssets(*Menu);
+		const FAssetData SkeletalMeshAsset = Context->SelectedAssets[0];
+		TSoftObjectPtr<USkeletalMesh> SkeletalMeshPtr = TSoftObjectPtr<USkeletalMesh>(SkeletalMeshAsset.ToSoftObjectPath());
+
+		FToolMenuSection& Section = Menu->FindOrAddSection("LODImport");
+		
+    	const int32 LODMax = GetNumberOfLODs(SkeletalMeshAsset);
+    	for(int32 LOD = 1; LOD <= LODMax; ++LOD)
+    	{
+    		const TAttribute<FText> Label = (LOD == LODMax) ? FText::Format(LOCTEXT("AddLODLevel", "Add LOD {0}"), FText::AsNumber(LOD)) : FText::Format( LOCTEXT("LODLevel", "Reimport LOD {0}"), FText::AsNumber( LOD ) );
+    		const TAttribute<FText> ToolTip = ( LOD == LODMax ) ? LOCTEXT("NewImportTip", "Import new LOD") : LOCTEXT("ReimportTip", "Reimport over existing LOD");
+    		const FSlateIcon Icon = FSlateIcon();
+
+    		FToolUIAction UIAction;
+    		UIAction.ExecuteAction = FToolMenuExecuteAction::CreateStatic(&ExecuteImportMeshLOD, SkeletalMeshPtr, LOD);
+    		Section.AddMenuEntry(*FString::Printf(TEXT("LOD_%d"), LOD), Label, ToolTip, Icon, UIAction);
+    	}
+    }
+    
+    void ExecuteNewSkeleton(const FToolMenuContext& InContext)
+    {
+		const UContentBrowserAssetContextMenuContext* CBContext = UContentBrowserAssetContextMenuContext::FindContextWithAssets(InContext);
+
+		IAssetTools::Get().CreateAssetsFrom<USkeletalMesh>(
+			CBContext->LoadSelectedObjects<USkeletalMesh>(), USkeleton::StaticClass(), TEXT("_Skeleton"), [](USkeletalMesh* SourceObject)
+			{
+				if (SourceObject->GetOutermost()->bIsCookedForEditor)
+				{
+					FAssetNotifications::CannotEditCookedAsset(SourceObject);
+					return static_cast<USkeletonFactory*>(nullptr);
+				}
+				
+				USkeletonFactory* Factory = NewObject<USkeletonFactory>();
+				Factory->TargetSkeletalMesh = SourceObject;
+				
+				return Factory;
+			}
+		);
+    }
+    
+    void ExecuteAssignSkeleton(const FToolMenuContext& InContext)
+    {
+		const UContentBrowserAssetContextMenuContext* CBContext = UContentBrowserAssetContextMenuContext::FindContextWithAssets(InContext);
+		
+		for (USkeletalMesh* SkeletalMesh : CBContext->LoadSelectedObjects<USkeletalMesh>())
+		{
+			UE::SkeletalMeshUtilities::AssignSkeletonToMesh(SkeletalMesh);
+    	}
+    }
+    
+    void ExecuteFindSkeleton(const FToolMenuContext& InContext)
+    {
+		const UContentBrowserAssetContextMenuContext* Context = UContentBrowserAssetContextMenuContext::FindContextWithAssets(InContext);
+
+    	TArray<UObject*> ObjectsToSync;
+    	for (USkeletalMesh* SkeletalMesh : Context->LoadSelectedObjects<USkeletalMesh>())
+    	{
+	        if (USkeleton* Skeleton = SkeletalMesh->GetSkeleton())
+    		{
+    			ObjectsToSync.AddUnique(Skeleton);
+    		}
+    	}
+    
+    	if ( ObjectsToSync.Num() > 0 )
+    	{
+    		FContentBrowserModule& ContentBrowserModule = FModuleManager::LoadModuleChecked<FContentBrowserModule>("ContentBrowser");
+    		ContentBrowserModule.Get().SyncBrowserToAssets(ObjectsToSync);
+    	}
+    }
+    
+    void FillSkeletonMenu(UToolMenu* Menu)
+    {
+		FToolMenuSection& Section = Menu->FindOrAddSection("SkeletonMenu");
+		Section.Label = LOCTEXT("SkeletonMenuHeading", "Skeleton");
+		
+		{
+			const TAttribute<FText> Label = LOCTEXT("SkeletalMesh_NewSkeleton", "Create Skeleton");
+			const TAttribute<FText> ToolTip = LOCTEXT("SkeletalMesh_NewSkeletonTooltip", "Creates a new skeleton for each of the selected meshes.");
+			const FSlateIcon Icon = FSlateIcon(FAppStyle::GetAppStyleSetName(), "AssetIcons.Skeleton");
+
+			FToolUIAction UIAction;
+			UIAction.ExecuteAction = FToolMenuExecuteAction::CreateStatic(&ExecuteNewSkeleton);
+			Section.AddMenuEntry("SkeletalMesh_NewSkeleton", Label, ToolTip, Icon, UIAction);
+		}
+		
+		{
+			const TAttribute<FText> Label = LOCTEXT("SkeletalMesh_AssignSkeleton", "Assign Skeleton");
+			const TAttribute<FText> ToolTip = LOCTEXT("SkeletalMesh_AssignSkeletonTooltip", "Assigns a skeleton to the selected meshes.");
+			const FSlateIcon Icon = FSlateIcon(FAppStyle::GetAppStyleSetName(), "Persona.AssetActions.AssignSkeleton");
+
+			FToolUIAction UIAction;
+			UIAction.ExecuteAction = FToolMenuExecuteAction::CreateStatic(&ExecuteAssignSkeleton);
+			Section.AddMenuEntry("SkeletalMesh_AssignSkeleton", Label, ToolTip, Icon, UIAction);
+		}
+    
+		{
+			const TAttribute<FText> Label = LOCTEXT("SkeletalMesh_FindSkeleton", "Find Skeleton");
+			const TAttribute<FText> ToolTip = LOCTEXT("SkeletalMesh_FindSkeletonTooltip", "Finds the skeleton used by the selected meshes in the content browser.");
+			const FSlateIcon Icon = FSlateIcon(FAppStyle::GetAppStyleSetName(), "Persona.AssetActions.FindSkeleton");
+
+			FToolUIAction UIAction;
+			UIAction.ExecuteAction = FToolMenuExecuteAction::CreateStatic(&ExecuteFindSkeleton);
+			Section.AddMenuEntry("SkeletalMesh_FindSkeleton", Label, ToolTip, Icon, UIAction);
+		}
+    }
+	
+    
+	static FDelayedAutoRegisterHelper DelayedAutoRegister(EDelayedRegisterRunPhase::EndOfEngineInit, []{ 
+    	UToolMenus::RegisterStartupCallback(FSimpleMulticastDelegate::FDelegate::CreateLambda([]()
+    	{
+			FToolMenuOwnerScoped OwnerScoped(UE_MODULE_NAME);
+			UToolMenu* Menu = UE::ContentBrowser::ExtendToolMenu_AssetContextMenu(USkeletalMesh::StaticClass());
+    			
+			FToolMenuSection& Section = Menu->FindOrAddSection("GetAssetActions");
+			Section.AddDynamicEntry(NAME_None, FNewToolMenuSectionDelegate::CreateLambda([](FToolMenuSection& InSection)
+			{
+				{
+					InSection.AddSubMenu(
+						"CreateSkeletalMeshSubmenu",
+						LOCTEXT("CreateSkeletalMeshSubmenu", "Create"),
+						LOCTEXT("CreateSkeletalMeshSubmenu_ToolTip", "Create related assets"),
+						FNewToolMenuDelegate::CreateStatic(&FillCreateMenu),
+						false,
+						FSlateIcon(FAppStyle::GetAppStyleSetName(), "Persona.AssetActions.CreateAnimAsset")
+					);
+				}
+
+				if (UContentBrowserAssetContextMenuContext::GetNumAssetsSelected(InSection))
+				{
+					InSection.AddSubMenu(
+						"SkeletalMesh_LODImport",	
+						LOCTEXT("SkeletalMesh_LODImport", "Import LOD"),
+						LOCTEXT("SkeletalMesh_LODImportTooltip", "Select which LODs to import."),
+						FNewToolMenuDelegate::CreateStatic(&GetLODMenu),
+						false,
+						FSlateIcon(FAppStyle::GetAppStyleSetName(), "Icons.LOD")
+					);
+				}
+
+				{
+					InSection.AddSubMenu(
+						"SkeletonSubmenu",
+						LOCTEXT("SkeletonSubmenu", "Skeleton"),
+						LOCTEXT("SkeletonSubmenu_ToolTip", "Skeleton related actions"),
+						FNewToolMenuDelegate::CreateStatic(&FillSkeletonMenu),
+						false,
+						FSlateIcon(FAppStyle::GetAppStyleSetName(), "Icons.SkeletalMesh")
+					);
+				}
+			}));
+		}));
+	});
+}
+
+#undef LOCTEXT_NAMESPACE
diff --git a/Engine/Plugins/Editor/EngineAssetDefinitions/Source/Private/Animation/AssetDefinition_SkeletalMesh.h b/Engine/Plugins/Editor/EngineAssetDefinitions/Source/Private/Animation/AssetDefinition_SkeletalMesh.h
new file mode 100644
--- /dev/null
+++ b/Engine/Plugins/Editor/EngineAssetDefinitions/Source/Private/Animation/AssetDefinition_SkeletalMesh.h
@@ -0,0 +1,39 @@
+// Copyright Epic Games, Inc. All Rights Reserved.
+
+#pragma once
+
+#include "CoreMinimal.h"
+#include "Engine/SkeletalMesh.h"
+#include "AssetDefinitionDefault.h"
+
+#include "AssetDefinition_SkeletalMesh.generated.h"
+
+UCLASS()
+class UAssetDefinition_SkeletalMesh : public UAssetDefinitionDefault
+{
+	GENERATED_BODY()
+
+public:
+	UAssetDefinition_SkeletalMesh();
+	
+	// UAssetDefinition Begin
+	virtual FText GetAssetDisplayName() const override { return NSLOCTEXT("AssetTypeActions", "AssetTypeActions_SkeletalMesh", "Skeletal Mesh"); }
+	virtual FLinearColor GetAssetColor() const override { return FLinearColor(FColor(241,163,241)); }
+	virtual TSoftClassPtr<UObject> GetAssetClass() const override { return USkeletalMesh::StaticClass(); }
+	virtual TConstArrayView<FAssetCategoryPath> GetAssetCategories() const override
+	{
+		static const auto Categories = { EAssetCategoryPaths::Basic, EAssetCategoryPaths::Animation };
+		return Categories;
+	}
+	virtual bool CanImport() const override { return true; }
+	virtual UThumbnailInfo* LoadThumbnailInfo(const FAssetData& InAssetData) const override;
+	virtual TSharedPtr<class SWidget> GetThumbnailOverlay(const FAssetData& AssetData) const override;
+	virtual EAssetCommandResult OpenAssets(const FAssetOpenArgs& OpenArgs) const override;
+	// UAssetDefinition End
+
+private:
+	void OnAssetRemoved(const FAssetData& AssetData) const;
+	EVisibility GetThumbnailSkinningOverlayVisibility(const FAssetData AssetData) const;
+
+	mutable TArray<FString> ThumbnailSkinningOverlayAssetNames;
+};
diff --git a/Engine/Plugins/Editor/EngineAssetDefinitions/Source/Private/Animation/AssetDefinition_Skeleton.cpp b/Engine/Plugins/Editor/EngineAssetDefinitions/Source/Private/Animation/AssetDefinition_Skeleton.cpp
new file mode 100644
--- /dev/null
+++ b/Engine/Plugins/Editor/EngineAssetDefinitions/Source/Private/Animation/AssetDefinition_Skeleton.cpp
@@ -0,0 +1,87 @@
+// Copyright Epic Games, Inc. All Rights Reserved.
+
+#include "Animation/AssetDefinition_Skeleton.h"
+#include "Widgets/Text/STextBlock.h"
+#include "ToolMenus.h"
+#include "Engine/SkeletalMesh.h"
+#include "Animation/AnimationAsset.h"
+#include "Animation/AnimSequenceBase.h"
+#include "Animation/AnimSequence.h"
+#include "Widgets/Layout/SBorder.h"
+#include "Misc/MessageDialog.h"
+#include "HAL/FileManager.h"
+#include "UObject/UObjectHash.h"
+#include "UObject/UObjectIterator.h"
+#include "Framework/Application/SlateApplication.h"
+#include "Widgets/Input/SButton.h"
+#include "Widgets/Input/SCheckBox.h"
+#include "Styling/AppStyle.h"
+#include "Components/SkeletalMeshComponent.h"
+#include "AnimationEditorUtils.h"
+#include "ISkeletonEditorModule.h"
+#include "Algo/Transform.h"
+#include "ContentBrowserMenuContexts.h"
+#include "Subsystems/AssetEditorSubsystem.h"
+#include "Editor.h"
+
+#define LOCTEXT_NAMESPACE "AssetTypeActions"
+
+EAssetCommandResult UAssetDefinition_Skeleton::OpenAssets(const FAssetOpenArgs& OpenArgs) const
+{
+	ISkeletonEditorModule& SkeletonEditorModule = FModuleManager::LoadModuleChecked<ISkeletonEditorModule>("SkeletonEditor");
+	
+	for (USkeleton* Skeleton : OpenArgs.LoadObjects<USkeleton>())
+	{
+		SkeletonEditorModule.CreateSkeletonEditor(OpenArgs.GetToolkitMode(), OpenArgs.ToolkitHost, Skeleton);
+	}
+
+	return EAssetCommandResult::Handled;
+}
+
+// Menu Extensions
+//--------------------------------------------------------------------
+
+namespace MenuExtension_Skeleton
+{
+	static void GetCreateAssetsForSkeletonMenu(UToolMenu* Menu)
+	{
+		const UContentBrowserAssetContextMenuContext* Context = UContentBrowserAssetContextMenuContext::FindContextWithAssets(*Menu);
+		TArray<TSoftObjectPtr<UObject>> Skeletons = Context->GetSelectedAssetSoftObjects<UObject>();
+
+		Menu->AddDynamicSection(NAME_None, FNewToolMenuDelegateLegacy::CreateLambda([Skeletons](FMenuBuilder& MenuBuilder, UToolMenu* Menu)
+		{
+			AnimationEditorUtils::FillCreateAssetMenu(MenuBuilder, Skeletons, FAnimAssetCreated::CreateLambda([](TArray<UObject*> NewAssets)
+				{
+					if(NewAssets.Num() > 1)
+					{
+						FModuleManager::LoadModuleChecked<FContentBrowserModule>("ContentBrowser").Get().SyncBrowserToAssets(NewAssets);
+					}
+					return true;
+				})
+			);
+		}));
+	}
+
+	static FDelayedAutoRegisterHelper DelayedAutoRegister(EDelayedRegisterRunPhase::EndOfEngineInit, []{ 
+		UToolMenus::RegisterStartupCallback(FSimpleMulticastDelegate::FDelegate::CreateLambda([]()
+		{
+			FToolMenuOwnerScoped OwnerScoped(UE_MODULE_NAME);
+			UToolMenu* Menu = UE::ContentBrowser::ExtendToolMenu_AssetContextMenu(USkeleton::StaticClass());
+		
+			FToolMenuSection& Section = Menu->FindOrAddSection("GetAssetActions");
+			Section.AddDynamicEntry(NAME_None, FNewToolMenuSectionDelegate::CreateLambda([](FToolMenuSection& InSection)
+			{
+				InSection.AddSubMenu(
+					"CreateSkeletonSubmenu",
+					LOCTEXT("CreateSkeletonSubmenu", "Create"),
+					LOCTEXT("CreateSkeletonSubmenu_ToolTip", "Create assets for this skeleton"),
+					FNewToolMenuDelegate::CreateStatic(&GetCreateAssetsForSkeletonMenu),
+					false, 
+					FSlateIcon(FAppStyle::GetAppStyleSetName(), "Persona.AssetActions.CreateAnimAsset")
+				);
+			}));
+		}));
+	});
+}
+
+#undef LOCTEXT_NAMESPACE
diff --git a/Engine/Plugins/Editor/EngineAssetDefinitions/Source/Private/Animation/AssetDefinition_Skeleton.h b/Engine/Plugins/Editor/EngineAssetDefinitions/Source/Private/Animation/AssetDefinition_Skeleton.h
new file mode 100644
--- /dev/null
+++ b/Engine/Plugins/Editor/EngineAssetDefinitions/Source/Private/Animation/AssetDefinition_Skeleton.h
@@ -0,0 +1,28 @@
+// Copyright Epic Games, Inc. All Rights Reserved.
+
+#pragma once
+
+#include "CoreMinimal.h"
+#include "Animation/Skeleton.h"
+#include "AssetDefinitionDefault.h"
+
+#include "AssetDefinition_Skeleton.generated.h"
+
+UCLASS()
+class UAssetDefinition_Skeleton : public UAssetDefinitionDefault
+{
+	GENERATED_BODY()
+
+public:
+	// UAssetDefinition Begin
+	virtual FText GetAssetDisplayName() const override { return NSLOCTEXT("AssetTypeActions", "AssetTypeActions_Skeleton", "Skeleton"); }
+	virtual FLinearColor GetAssetColor() const override { return FLinearColor(FColor(105,181,205)); }
+	virtual TSoftClassPtr<UObject> GetAssetClass() const override { return USkeleton::StaticClass(); }
+	virtual TConstArrayView<FAssetCategoryPath> GetAssetCategories() const override
+	{
+		static const auto Categories = { EAssetCategoryPaths::Animation };
+		return Categories;
+	}
+	virtual EAssetCommandResult OpenAssets(const FAssetOpenArgs& OpenArgs) const override;
+	// UAssetDefinition End
+};
diff --git a/Engine/Plugins/Editor/EngineAssetDefinitions/Source/Private/Animation/SkeletalMeshUtilities.cpp b/Engine/Plugins/Editor/EngineAssetDefinitions/Source/Private/Animation/SkeletalMeshUtilities.cpp
new file mode 100644
--- /dev/null
+++ b/Engine/Plugins/Editor/EngineAssetDefinitions/Source/Private/Animation/SkeletalMeshUtilities.cpp
@@ -0,0 +1,513 @@
+// Copyright Epic Games, Inc. All Rights Reserved.
+
+#include "Animation/SkeletalMeshUtilities.h"
+#include "Animation/Skeleton.h"
+#include "Widgets/Text/STextBlock.h"
+#include "ToolMenus.h"
+#include "Misc/PackageName.h"
+#include "Widgets/Layout/SBorder.h"
+#include "Misc/MessageDialog.h"
+#include "SlateOptMacros.h"
+#include "Framework/Application/SlateApplication.h"
+#include "Widgets/Layout/SSeparator.h"
+#include "Widgets/Layout/SUniformGridPanel.h"
+#include "Widgets/Input/SButton.h"
+#include "Widgets/Layout/SScrollBox.h"
+#include "Widgets/Input/SCheckBox.h"
+#include "Styling/AppStyle.h"
+#include "Factories/SkeletonFactory.h"
+#include "EditorFramework/AssetImportData.h"
+#include "ThumbnailRendering/SceneThumbnailInfo.h"
+#include "AssetRegistry/AssetRegistryModule.h"
+#include "SSkeletonWidget.h"
+#include "ContentBrowserModule.h"
+#include "AnimationEditorUtils.h"
+#include "FbxMeshUtils.h"
+#include "AssetNotifications.h"
+#include "PhysicsEngine/PhysicsAsset.h"
+#include "ISkeletalMeshEditorModule.h"
+#include "Algo/Transform.h"
+#include "Factories/PhysicsAssetFactory.h"
+#include "EditorReimportHandler.h"
+#include "Widgets/Images/SImage.h"
+#include "Factories/FbxSkeletalMeshImportData.h"
+#include "Subsystems/AssetEditorSubsystem.h"
+#include "Editor.h"
+
+#define LOCTEXT_NAMESPACE "AssetTypeActions"
+
+namespace UE::SkeletalMeshUtilities
+{
+	/**
+	* FBoneCheckboxInfo
+	* 
+	* Context data for the SDlgMergeSkeleton panel check boxes
+	*/
+	struct FBoneCheckboxInfo
+	{
+		FName BoneName;
+		int32 BoneID;
+		bool  bUsed;
+	};
+
+	/** 
+	* FDlgMergeSkeleton
+	* 
+	* Wrapper class for SDlgMergeSkeleton. This class creates and launches a dialog then awaits the
+	* result to return to the user. 
+	*/
+	class FDlgMergeSkeleton
+	{
+	public:
+		enum EResult
+		{
+			Cancel = 0,			// No/Cancel, normal usage would stop the current action
+			Confirm = 1,		// Yes/Ok/Etc, normal usage would continue with action
+		};
+
+		FDlgMergeSkeleton( USkeletalMesh* InMesh, USkeleton* InSkeleton );
+
+		/**  Shows the dialog box and waits for the user to respond. */
+		EResult ShowModal();
+
+		// List of required bones for skeleton
+		TArray<int32> RequiredBones;
+	private:
+
+		/** Cached pointer to the modal window */
+		TSharedPtr<SWindow> DialogWindow;
+
+		/** Cached pointer to the merge skeleton widget */
+		TSharedPtr<class SDlgMergeSkeleton> DialogWidget;
+
+		/** The SkeletalMesh to merge bones from*/
+		USkeletalMesh* Mesh;
+		/** The Skeleton to merge bones to */
+		USkeleton* Skeleton;
+	};
+
+	/** 
+	 * Slate panel for choosing which bones to merge into the skeleton
+	 */
+	class SDlgMergeSkeleton : public SCompoundWidget
+	{
+	public:
+
+		SLATE_BEGIN_ARGS(SDlgMergeSkeleton)
+			{}
+			/** Window in which this widget resides */
+			SLATE_ATTRIBUTE(TSharedPtr<SWindow>, ParentWindow)
+		SLATE_END_ARGS()	
+
+		/**
+		 * Constructs this widget
+		 *
+		 * @param	InArgs	The declaration data for this widget
+		 */
+		BEGIN_SLATE_FUNCTION_BUILD_OPTIMIZATION
+		void Construct( const FArguments& InArgs )
+		{
+			UserResponse = FDlgMergeSkeleton::Cancel;
+			ParentWindow = InArgs._ParentWindow.Get();
+
+			this->ChildSlot[
+				SNew(SVerticalBox)
+				+SVerticalBox::Slot()
+				.AutoHeight()
+				.Padding(8.0f, 4.0f, 8.0f, 4.0f)
+				[
+					SNew( STextBlock )
+					.Text( LOCTEXT("MergeSkeletonDlgDescription", "Would you like to add following bones to the skeleton?"))
+				]
+				+SVerticalBox::Slot()
+				.AutoHeight()
+				.Padding(8.0f, 4.0f, 8.0f, 4.0f)
+				[
+					SNew(SSeparator)
+				]
+				+SVerticalBox::Slot()
+				.Padding(8.0f, 4.0f, 8.0f, 4.0f)
+				[
+					SNew(SBorder)
+					[
+						SNew(SScrollBox)
+						+SScrollBox::Slot()
+						[
+							//Save this widget so we can populate it later with check boxes
+							SAssignNew(CheckBoxContainer, SVerticalBox)
+						]
+					]
+				]
+				+SVerticalBox::Slot()
+				.AutoHeight()
+				.HAlign(HAlign_Left)
+				.Padding(8.0f, 4.0f, 8.0f, 4.0f)
+				[
+					SNew(SUniformGridPanel)
+					.SlotPadding(FAppStyle::GetMargin("StandardDialog.SlotPadding"))
+					.MinDesiredSlotWidth(FAppStyle::GetFloat("StandardDialog.MinDesiredSlotWidth"))
+					.MinDesiredSlotHeight(FAppStyle::GetFloat("StandardDialog.MinDesiredSlotHeight"))
+					+SUniformGridPanel::Slot(0,0)
+					[
+						SNew(SButton) 
+						.HAlign(HAlign_Center)
+						.ContentPadding(FAppStyle::GetMargin("StandardDialog.ContentPadding"))
+						.OnClicked(this, &SDlgMergeSkeleton::ChangeAllOptions, true)
+						.Text(LOCTEXT("SkeletonMergeSelectAll", "Select All"))
+					]
+					+SUniformGridPanel::Slot(1,0)
+					[
+						SNew(SButton) 
+						.HAlign(HAlign_Center)
+						.ContentPadding(FAppStyle::GetMargin("StandardDialog.ContentPadding"))
+						.OnClicked(this, &SDlgMergeSkeleton::ChangeAllOptions, false)
+						.Text(LOCTEXT("SkeletonMergeDeselectAll", "Deselect All"))
+					]
+				]
+				+SVerticalBox::Slot()
+				.AutoHeight()
+				.Padding(8.0f, 4.0f, 8.0f, 4.0f)
+				[
+					SNew(SSeparator)
+				]
+				+SVerticalBox::Slot()
+				.AutoHeight()
+				.HAlign(HAlign_Right)
+				.Padding(8.0f, 4.0f, 8.0f, 4.0f)
+				[
+					SNew(SUniformGridPanel)
+					.SlotPadding(FAppStyle::GetMargin("StandardDialog.SlotPadding"))
+					.MinDesiredSlotWidth(FAppStyle::GetFloat("StandardDialog.MinDesiredSlotWidth"))
+					.MinDesiredSlotHeight(FAppStyle::GetFloat("StandardDialog.MinDesiredSlotHeight"))
+					+SUniformGridPanel::Slot(0,0)
+					[
+						SNew(SButton) 
+						.HAlign(HAlign_Center)
+						.ContentPadding(FAppStyle::GetMargin("StandardDialog.ContentPadding"))
+						.OnClicked( this, &SDlgMergeSkeleton::OnButtonClick, FDlgMergeSkeleton::Confirm )
+						.Text(LOCTEXT("SkeletonMergeOk", "OK"))
+					]
+					+SUniformGridPanel::Slot(1,0)
+					[
+						SNew(SButton) 
+						.HAlign(HAlign_Center)
+						.ContentPadding(FAppStyle::GetMargin("StandardDialog.ContentPadding"))
+						.OnClicked( this, &SDlgMergeSkeleton::OnButtonClick, FDlgMergeSkeleton::Cancel )
+						.Text(LOCTEXT("SkeletonMergeCancel", "Cancel"))
+					]
+				]
+			];
+		}
+		END_SLATE_FUNCTION_BUILD_OPTIMIZATION
+
+		/**
+		 * Creates a Slate check box 
+		 *
+		 * @param	Label		Text label for the check box
+		 * @param	ButtonId	The ID for the check box
+		 * @return				The created check box widget
+		 */
+		TSharedRef<SWidget> CreateCheckBox( const FString& Label, int32 ButtonId )
+		{
+			return
+				SNew(SCheckBox)
+				.IsChecked( this, &SDlgMergeSkeleton::IsCheckboxChecked, ButtonId )
+				.OnCheckStateChanged( this, &SDlgMergeSkeleton::OnCheckboxChanged, ButtonId )
+				[
+					SNew(STextBlock).Text(FText::FromString(Label))
+				];
+		}
+
+		/**
+		 * Returns the state of the check box
+		 *
+		 * @param	ButtonId	The ID for the check box
+		 * @return				The status of the check box
+		 */
+		ECheckBoxState IsCheckboxChecked( int32 ButtonId ) const
+		{
+			return CheckBoxInfoMap.FindChecked(ButtonId).bUsed ? ECheckBoxState::Checked : ECheckBoxState::Unchecked;
+		}
+
+		/**
+		 * Handler for all check box clicks
+		 *
+		 * @param	NewCheckboxState	The new state of the check box
+		 * @param	CheckboxThatChanged	The ID of the radio button that has changed. 
+		 */
+		void OnCheckboxChanged( ECheckBoxState NewCheckboxState, int32 CheckboxThatChanged )
+		{
+			FBoneCheckboxInfo& Info = CheckBoxInfoMap.FindChecked(CheckboxThatChanged);
+			Info.bUsed = !Info.bUsed;
+		}
+
+		/**
+		 * Handler for the Select All and Deselect All buttons
+		 *
+		 * @param	bNewCheckedState	The new state of the check boxes
+		 */
+		FReply ChangeAllOptions(bool bNewCheckedState)
+		{
+			for(auto Iter = CheckBoxInfoMap.CreateIterator(); Iter; ++Iter)
+			{
+				FBoneCheckboxInfo& Info = Iter.Value();
+				Info.bUsed = bNewCheckedState;
+			}
+			return FReply::Handled();
+		}
+
+		/**
+		 * Populated the dialog with multiple check boxes, each corresponding to a bone
+		 *
+		 * @param	BoneInfos	The list of Bones to populate the dialog with
+		 */
+		void PopulateOptions(TArray<FBoneCheckboxInfo>& BoneInfos)
+		{
+			for(auto Iter = BoneInfos.CreateIterator(); Iter; ++Iter)
+			{
+				FBoneCheckboxInfo& Info = (*Iter);
+				Info.bUsed = true;
+
+				CheckBoxInfoMap.Add(Info.BoneID, Info);
+
+				CheckBoxContainer->AddSlot()
+				.AutoHeight()
+				[
+					CreateCheckBox(Info.BoneName.GetPlainNameString(), Info.BoneID)
+				];
+			}
+		}
+
+		/** 
+		 * Returns the EResult of the button which the user pressed. Closing of the dialog
+		 * in any other way than clicking "Ok" results in this returning a "Cancel" value
+		 */
+		FDlgMergeSkeleton::EResult GetUserResponse() const 
+		{
+			return UserResponse; 
+		}
+
+		/** 
+		 * Returns whether the user selected that bone to be used (checked its respective check box)
+		 */
+		bool IsBoneIncluded(int32 BoneID)
+		{
+			auto* Item = CheckBoxInfoMap.Find(BoneID);
+			return Item ? Item->bUsed : false;
+		}
+
+	private:
+		
+		/** 
+		 * Handles when a button is pressed, should be bound with appropriate EResult Key
+		 * 
+		 * @param ButtonID - The return type of the button which has been pressed.
+		 */
+		FReply OnButtonClick(FDlgMergeSkeleton::EResult ButtonID)
+		{
+			ParentWindow->RequestDestroyWindow();
+			UserResponse = ButtonID;
+
+			return FReply::Handled();
+		}
+
+		/** Stores the users response to this dialog */
+		FDlgMergeSkeleton::EResult	 UserResponse;
+
+		/** The slate container that the bone check boxes get added to */
+		TSharedPtr<SVerticalBox>	 CheckBoxContainer;
+		/** Store the check box state for each bone */
+		TMap<int32,FBoneCheckboxInfo> CheckBoxInfoMap;
+
+		/** Pointer to the window which holds this Widget, required for modal control */
+		TSharedPtr<SWindow>			 ParentWindow;
+	};
+
+	FDlgMergeSkeleton::FDlgMergeSkeleton( USkeletalMesh* InMesh, USkeleton* InSkeleton )
+	{
+		Mesh = InMesh;
+		Skeleton = InSkeleton;
+
+		if (FSlateApplication::IsInitialized())
+		{
+			DialogWindow = SNew(SWindow)
+				.Title( LOCTEXT("MergeSkeletonDlgTitle", "Merge Bones") )
+				.SupportsMinimize(false) .SupportsMaximize(false)
+				.ClientSize(FVector2D(350, 500));
+
+			TSharedPtr<SBorder> DialogWrapper = 
+				SNew(SBorder)
+				.BorderImage(FAppStyle::GetBrush("ToolPanel.GroupBorder"))
+				.Padding(4.0f)
+				[
+					SAssignNew(DialogWidget, SDlgMergeSkeleton)
+					.ParentWindow(DialogWindow)
+				];
+
+			DialogWindow->SetContent(DialogWrapper.ToSharedRef());
+		}
+	}
+
+	FDlgMergeSkeleton::EResult FDlgMergeSkeleton::ShowModal()
+	{
+		RequiredBones.Empty();
+
+		TMap<FName, int32> BoneIndicesMap;
+		TArray<FBoneCheckboxInfo> BoneInfos;
+
+		// Make a list of all skeleton bone list
+		const FReferenceSkeleton& RefSkeleton = Skeleton->GetReferenceSkeleton();
+		for ( int32 BoneTreeId=0; BoneTreeId<RefSkeleton.GetRawBoneNum(); ++BoneTreeId )
+		{
+			const FName& BoneName = RefSkeleton.GetBoneName(BoneTreeId);
+			BoneIndicesMap.Add(BoneName, BoneTreeId);
+		}
+
+		for ( int32 RefBoneId=0 ; RefBoneId< Mesh->GetRefSkeleton().GetRawBoneNum() ; ++RefBoneId )
+		{
+			const FName& BoneName = Mesh->GetRefSkeleton().GetBoneName(RefBoneId);
+			// if I can't find this from Skeleton
+			if (BoneIndicesMap.Find(BoneName)==NULL)
+			{
+				FBoneCheckboxInfo Info;
+				Info.BoneID = RefBoneId;
+				Info.BoneName = BoneName;
+				BoneInfos.Add(Info);
+			}
+		}
+
+		if (BoneInfos.Num() == 0)
+		{
+			// it's all identical, but still need to return RequiredBones
+			// for the case, where they'd like to replace the one exactly same hierarchy but different skeleton 
+			for ( int32 RefBoneId= 0 ; RefBoneId< Mesh->GetRefSkeleton().GetRawBoneNum() ; ++RefBoneId )
+			{
+				const FName& BoneName = Mesh->GetRefSkeleton().GetBoneName(RefBoneId);
+				RequiredBones.Add(RefBoneId);
+			}
+
+			return EResult::Confirm;
+		}
+
+		DialogWidget->PopulateOptions(BoneInfos);
+
+		//Show Dialog
+		GEditor->EditorAddModalWindow(DialogWindow.ToSharedRef());
+		EResult UserResponse = (EResult)DialogWidget->GetUserResponse();
+
+		if(UserResponse == EResult::Confirm)
+		{
+			for ( int32 RefBoneId= 0 ; RefBoneId< Mesh->GetRefSkeleton().GetRawBoneNum() ; ++RefBoneId )
+			{
+				if ( DialogWidget->IsBoneIncluded(RefBoneId) )
+				{
+					TArray<int32> ParentList;
+					
+					// I need to make sure parent exists first
+					int32 ParentIndex = Mesh->GetRefSkeleton().GetParentIndex(RefBoneId);
+					
+					// make sure RequiredBones already have ParentIndex
+					while (ParentIndex >= 0 )
+					{
+						// if I don't have it yet
+						if ( RequiredBones.Contains(ParentIndex)==false )
+						{
+							ParentList.Add(ParentIndex);
+						}
+
+						ParentIndex = Mesh->GetRefSkeleton().GetParentIndex(ParentIndex);
+					}
+
+					if ( ParentList.Num() > 0 )
+					{
+						// if we need to add parent list
+						// add from back to front (since it's added from child to up
+						for (int32 I=ParentList.Num()-1; I>=0; --I)
+						{
+							RequiredBones.Add(ParentList[I]);
+						}
+					}
+
+					RequiredBones.Add(RefBoneId);
+				}
+			}
+		}
+		return UserResponse;
+	}
+
+	void AssignSkeletonToMesh(USkeletalMesh* SkelMesh)
+	{
+		if (!SkelMesh || SkelMesh->GetOutermost()->bIsCookedForEditor)
+		{
+			FAssetNotifications::CannotEditCookedAsset(SkelMesh);
+			return;
+		}
+
+		// Create a skeleton asset from the selected skeletal mesh. Defaults to being in the same package/group as the skeletal mesh.
+		TSharedRef<SWindow> WidgetWindow = SNew(SWindow)
+			.Title(LOCTEXT("ChooseSkeletonWindowTitle", "Choose Skeleton"))
+			.ClientSize(FVector2D(400,600));
+		TSharedPtr<SSkeletonSelectorWindow> SkeletonSelectorWindow;
+		WidgetWindow->SetContent
+			(
+				SNew(SBorder)
+					.BorderImage(FAppStyle::GetBrush("ToolPanel.GroupBorder"))
+				[
+					SAssignNew(SkeletonSelectorWindow, SSkeletonSelectorWindow)
+						.Object(SkelMesh)
+						.WidgetWindow(WidgetWindow)
+				]
+			);
+
+		GEditor->EditorAddModalWindow(WidgetWindow);
+		USkeleton* SelectedSkeleton = SkeletonSelectorWindow->GetSelectedSkeleton();
+
+		// only do this if not same
+		if ( SelectedSkeleton )
+		{
+			FDlgMergeSkeleton AssetDlg( SkelMesh, SelectedSkeleton );
+			if(AssetDlg.ShowModal() == FDlgMergeSkeleton::Confirm)
+			{			
+				TArray<int32> RequiredBones;
+				RequiredBones.Append(AssetDlg.RequiredBones);
+
+				if ( RequiredBones.Num() > 0 )
+				{
+					// Do automatic asset generation.
+					bool bSuccess = SelectedSkeleton->MergeBonesToBoneTree( SkelMesh, RequiredBones );
+					if ( bSuccess )
+					{
+						if (SkelMesh->GetSkeleton() != SelectedSkeleton)
+						{
+							SkelMesh->SetSkeleton(SelectedSkeleton);
+							SkelMesh->MarkPackageDirty();
+						}
+						FAssetNotifications::SkeletonNeedsToBeSaved(SelectedSkeleton);
+					}
+					else
+					{
+						// if failed, ask if user would like to regenerate skeleton hierarchy
+						if ( EAppReturnType::Yes == FMessageDialog::Open( EAppMsgType::YesNo, 
+							LOCTEXT("SkeletonMergeBones_Override", "FAILED TO MERGE BONES:  \n\nThis could happen if significant hierarchical changes have been made,\ne.g. inserting a bone between nodes.\nWould you like to regenerate the skeleton from this mesh? \n\n***WARNING: THIS WILL INVALIDATE ALL ANIMATION DATA THAT IS LINKED TO THIS SKELETON***\n")))
+						{
+							if ( SelectedSkeleton->RecreateBoneTree( SkelMesh ) )
+							{
+								FAssetNotifications::SkeletonNeedsToBeSaved( SelectedSkeleton );
+							}
+						}
+						else
+						{
+							FMessageDialog::Open( EAppMsgType::Ok, LOCTEXT("SkeletonMergeBonesFailure", "Failed to merge bones to Skeleton") );
+						}
+					}
+				}
+				else
+				{
+					FMessageDialog::Open( EAppMsgType::Ok, LOCTEXT("SkeletonMergeBonesFailure", "Failed to merge bones to Skeleton") );
+				}
+			}
+		}
+	}
+}
+
+#undef LOCTEXT_NAMESPACE
\ No newline at end of file
diff --git a/Engine/Plugins/Editor/EngineAssetDefinitions/Source/Private/Animation/SkeletalMeshUtilities.h b/Engine/Plugins/Editor/EngineAssetDefinitions/Source/Private/Animation/SkeletalMeshUtilities.h
new file mode 100644
--- /dev/null
+++ b/Engine/Plugins/Editor/EngineAssetDefinitions/Source/Private/Animation/SkeletalMeshUtilities.h
@@ -0,0 +1,12 @@
+// Copyright Epic Games, Inc. All Rights Reserved.
+
+#pragma once
+
+#include "CoreMinimal.h"
+
+class USkeletalMesh;
+
+namespace UE::SkeletalMeshUtilities
+{
+	void AssignSkeletonToMesh(USkeletalMesh* SkelMesh);
+}
\ No newline at end of file
diff --git a/Engine/Plugins/Editor/EngineAssetDefinitions/Source/Private/AssetDefinition_Font.h b/Engine/Plugins/Editor/EngineAssetDefinitions/Source/Private/AssetDefinition_Font.h
--- a/Engine/Plugins/Editor/EngineAssetDefinitions/Source/Private/AssetDefinition_Font.h
+++ b/Engine/Plugins/Editor/EngineAssetDefinitions/Source/Private/AssetDefinition_Font.h
@@ -25,11 +25,4 @@ public:
 	
 	virtual EAssetCommandResult OpenAssets(const FAssetOpenArgs& OpenArgs) const override;
 	// UAssetDefinition End
-
-private:
-	/** Can we execute a reimport for the selected objects? */
-	bool CanExecuteReimport(const TArray<TWeakObjectPtr<UFont>> Objects) const;
-
-	/** Handler for when Reimport is selected */
-	void ExecuteReimport(const TArray<TWeakObjectPtr<UFont>> Objects) const;
 };
diff --git a/Engine/Source/Developer/AssetTools/Private/AssetTools.cpp b/Engine/Source/Developer/AssetTools/Private/AssetTools.cpp
--- a/Engine/Source/Developer/AssetTools/Private/AssetTools.cpp
+++ b/Engine/Source/Developer/AssetTools/Private/AssetTools.cpp
@@ -37,7 +37,6 @@
 #include "ClassViewerFilter.h"
 #include "AssetTypeActions/AssetTypeActions_Blueprint.h"
 #include "AssetTypeActions/AssetTypeActions_BlueprintGeneratedClass.h"
-#include "AssetTypeActions/AssetTypeActions_SkeletalMesh.h"
 #include "AssetTypeActions/AssetTypeActions_AnimationAsset.h"
 #include "AssetTypeActions/AssetTypeActions_AnimBlueprint.h"
 #include "AssetTypeActions/AssetTypeActions_AnimBlueprintInterface.h"
@@ -63,7 +62,6 @@
 #include "AssetTypeActions/AssetTypeActions_MirrorDataTable.h"
 #include "AssetTypeActions/AssetTypeActions_ParticleSystem.h"
 #include "AssetTypeActions/AssetTypeActions_PhysicalMaterialMask.h"
-#include "AssetTypeActions/AssetTypeActions_Skeleton.h"
 #include "WorldPartition/WorldPartition.h"
 #include "SDiscoveringAssetsDialog.h"
 #include "AssetFixUpRedirectors.h"
@@ -1354,8 +1352,6 @@ UAssetToolsImpl::UAssetToolsImpl(const FObjectInitializer& ObjectInitializer)
 	RegisterAssetTypeActions(MakeShareable(new FAssetTypeActions_MirrorDataTable));
 	RegisterAssetTypeActions(MakeShareable(new FAssetTypeActions_ParticleSystem));
 	RegisterAssetTypeActions(MakeShareable(new FAssetTypeActions_PhysicalMaterialMask));
-	RegisterAssetTypeActions(MakeShareable(new FAssetTypeActions_SkeletalMesh));
-	RegisterAssetTypeActions(MakeShareable(new FAssetTypeActions_Skeleton));
 
 	// Note: Please don't add any more actions here!  They belong in an editor-only module that is more tightly
 	// coupled to your new system, and you should not create a dependency on your new system from AssetTools.
@@ -1657,15 +1653,16 @@ void UAssetToolsImpl::CreateAssetsFrom(TConstArrayView<UObject*> SourceObjects,
 	{
 		if (UObject* SourceObject = SourceObjects[0])
 		{
-			// Create an appropriate and unique name 
-			FString Name;
-			FString PackageName;
-			CreateUniqueAssetName(SourceObject->GetOutermost()->GetName(), DefaultSuffix, PackageName, Name);
-
-			UFactory* Factory = FactoryConstructor(SourceObject);
+			if (UFactory* Factory = FactoryConstructor(SourceObject))
+			{
+				// Create an appropriate and unique name 
+				FString Name;
+				FString PackageName;
+				CreateUniqueAssetName(SourceObject->GetOutermost()->GetName(), DefaultSuffix, PackageName, Name);
 
-			FContentBrowserModule& ContentBrowserModule = FModuleManager::LoadModuleChecked<FContentBrowserModule>("ContentBrowser");
-			ContentBrowserModule.Get().CreateNewAsset(Name, FPackageName::GetLongPackagePath(PackageName), CreateAssetType, Factory);
+				FContentBrowserModule& ContentBrowserModule = FModuleManager::LoadModuleChecked<FContentBrowserModule>("ContentBrowser");
+				ContentBrowserModule.Get().CreateNewAsset(Name, FPackageName::GetLongPackagePath(PackageName), CreateAssetType, Factory);
+			}
 		}
 	}
 	else
@@ -1675,17 +1672,18 @@ void UAssetToolsImpl::CreateAssetsFrom(TConstArrayView<UObject*> SourceObjects,
 		{
 			if ( SourceObject )
 			{
-				// Determine an appropriate name
-				FString Name;
-				FString PackageName;
-				CreateUniqueAssetName(SourceObject->GetOutermost()->GetName(), DefaultSuffix, PackageName, Name);
-
 				// Create the factory used to generate the asset
-				UFactory* Factory = FactoryConstructor(SourceObject);
-					
-				if (UObject* NewAsset = CreateAsset(Name, FPackageName::GetLongPackagePath(PackageName), CreateAssetType, Factory, CallingContext))
+				if (UFactory* Factory = FactoryConstructor(SourceObject))
 				{
-					ObjectsToSync.Add(NewAsset);
+					// Determine an appropriate name
+					FString Name;
+					FString PackageName;
+					CreateUniqueAssetName(SourceObject->GetOutermost()->GetName(), DefaultSuffix, PackageName, Name);
+					
+					if (UObject* NewAsset = CreateAsset(Name, FPackageName::GetLongPackagePath(PackageName), CreateAssetType, Factory, CallingContext))
+					{
+						ObjectsToSync.Add(NewAsset);
+					}
 				}
 			}
 		}
diff --git a/Engine/Source/Developer/AssetTools/Private/AssetTypeActions/AssetTypeActions_SkeletalMesh.cpp b/Engine/Source/Developer/AssetTools/Private/AssetTypeActions/AssetTypeActions_SkeletalMesh.cpp
deleted file mode 100644
diff --git a/Engine/Source/Developer/AssetTools/Private/AssetTypeActions/AssetTypeActions_SkeletalMesh.h b/Engine/Source/Developer/AssetTools/Private/AssetTypeActions/AssetTypeActions_SkeletalMesh.h
deleted file mode 100644
diff --git a/Engine/Source/Developer/AssetTools/Private/AssetTypeActions/AssetTypeActions_Skeleton.cpp b/Engine/Source/Developer/AssetTools/Private/AssetTypeActions/AssetTypeActions_Skeleton.cpp
deleted file mode 100644
diff --git a/Engine/Source/Developer/AssetTools/Private/AssetTypeActions/AssetTypeActions_Skeleton.h b/Engine/Source/Developer/AssetTools/Private/AssetTypeActions/AssetTypeActions_Skeleton.h
deleted file mode 100644
diff --git a/Engine/Source/Developer/AssetTools/Public/AssetTypeActions_Base.h b/Engine/Source/Developer/AssetTools/Public/AssetTypeActions_Base.h
--- a/Engine/Source/Developer/AssetTools/Public/AssetTypeActions_Base.h
+++ b/Engine/Source/Developer/AssetTools/Public/AssetTypeActions_Base.h
@@ -218,6 +218,20 @@ protected:
 
 		return TypedObjects;
 	}
+	
+	template <typename T>
+    static TArray<TSoftObjectPtr<T>> GetTypedSoftObjectPtrs(const TArray<UObject*>& InObjects)
+    {
+    	check(InObjects.Num() > 0);
+
+    	TArray<TSoftObjectPtr<T>> TypedObjects;
+    	for (auto ObjIt = InObjects.CreateConstIterator(); ObjIt; ++ObjIt)
+    	{
+    		TypedObjects.Add( CastChecked<T>(*ObjIt) );
+    	}
+
+    	return TypedObjects;
+    }
 
 	template <typename T>
 	static TArray<T*> GetTypedObjectPtrs(const TArray<UObject*>& InObjects)
diff --git a/Engine/Source/Developer/AssetTools/Public/IAssetTools.h b/Engine/Source/Developer/AssetTools/Public/IAssetTools.h
--- a/Engine/Source/Developer/AssetTools/Public/IAssetTools.h
+++ b/Engine/Source/Developer/AssetTools/Public/IAssetTools.h
@@ -301,9 +301,19 @@ public:
 	 * it needed a utility.  In the case that you only have a single SourceObject, we'll lean on the content browser
 	 * to create the asset and focus you to it so you can rename it inline.  However in the case that multiple assets
 	 * get created we'll construct each one and then sync the content browser to them.
+	 *
+	 * You can return null for a factory if you need to skip a given SourceObject.
 	 */
 	virtual void CreateAssetsFrom(TConstArrayView<UObject*> SourceObjects, UClass* CreateAssetType, const FString& DefaultSuffix, TFunctionRef<UFactory*(UObject*)> FactoryConstructor, FName CallingContext = NAME_None) = 0;
 
+	/**
+	 * Creates one or more assets using the source objects as the basis for the next type.  This is a common enough operation
+	 * it needed a utility.  In the case that you only have a single SourceObject, we'll lean on the content browser
+	 * to create the asset and focus you to it so you can rename it inline.  However in the case that multiple assets
+	 * get created we'll construct each one and then sync the content browser to them.
+	 *
+	 * You can return null for a factory if you need to skip a given SourceObject.
+	 */
 	template<typename SourceObjectType, typename = typename TEnableIf<TIsDerivedFrom<SourceObjectType, UObject>::Value>::Type>
 	void CreateAssetsFrom(TConstArrayView<SourceObjectType*> SourceObjects, UClass* CreateAssetType, const FString& DefaultSuffix, TFunctionRef<UFactory*(SourceObjectType*)> FactoryConstructor, FName CallingContext = NAME_None)
 	{
diff --git a/Engine/Source/Editor/AnimationEditor/Private/AnimationEditor.cpp b/Engine/Source/Editor/AnimationEditor/Private/AnimationEditor.cpp
--- a/Engine/Source/Editor/AnimationEditor/Private/AnimationEditor.cpp
+++ b/Engine/Source/Editor/AnimationEditor/Private/AnimationEditor.cpp
@@ -670,7 +670,7 @@ void FAnimationEditor::OnExportToFBX(const EExportSourceOption Option)
 	}
 	else if (Option == EExportSourceOption::CurrentAnimation_PreviewMesh)
 	{
-		TArray<TWeakObjectPtr<UObject>> Skeletons;
+		TArray<TSoftObjectPtr<UObject>> Skeletons;
 		Skeletons.Add(PersonaToolkit->GetSkeleton());
 		AnimationEditorUtils::CreateAnimationAssets(Skeletons, UAnimSequence::StaticClass(), FString("_PreviewMesh"), FAnimAssetCreated::CreateSP(this, &FAnimationEditor::ExportToFBX, true), AnimationAsset, true);
 	}
diff --git a/Engine/Source/Editor/ContentBrowser/Public/ContentBrowserMenuContexts.h b/Engine/Source/Editor/ContentBrowser/Public/ContentBrowserMenuContexts.h
--- a/Engine/Source/Editor/ContentBrowser/Public/ContentBrowserMenuContexts.h
+++ b/Engine/Source/Editor/ContentBrowser/Public/ContentBrowserMenuContexts.h
@@ -185,6 +185,32 @@ public:
 	 */
 	const FAssetData* GetSingleSelectedAssetOfType(const UClass* AssetClass, EIncludeSubclasses IncludeSubclasses = EIncludeSubclasses::Yes) const;
 
+	/**
+	 * Get the selected assets as an array of SoftObjectPtr<T>'s.
+	 */
+	template<typename ExpectedAssetType>
+	TArray<TSoftObjectPtr<ExpectedAssetType>> GetSelectedAssetSoftObjects(EIncludeSubclasses IncludeSubclasses = EIncludeSubclasses::Yes) const
+	{
+		TArray<TSoftObjectPtr<ExpectedAssetType>> Result;
+		Result.Reserve(SelectedAssets.Num());
+		if (IncludeSubclasses == EIncludeSubclasses::Yes)
+		{
+			Algo::TransformIf(SelectedAssets, Result,
+				[](const FAssetData& AssetData) { return AssetData.IsInstanceOf(ExpectedAssetType::StaticClass()); },
+				[](const FAssetData& AssetData) { return TSoftObjectPtr<ExpectedAssetType>(AssetData.ToSoftObjectPath()); }
+			);
+		}
+		else
+		{
+			Algo::TransformIf(SelectedAssets, Result,
+				[](const FAssetData& AssetData) { return AssetData.GetClass() == ExpectedAssetType::StaticClass(); },
+				[](const FAssetData& AssetData) { return TSoftObjectPtr<ExpectedAssetType>(AssetData.ToSoftObjectPath()); }
+			);
+		}
+
+		return Result;
+	}
+
 	/**
 	 * Finds the Content Browser MenuContext from a Menu or Section, and returns the context provided there are some
 	 * selected assets.
diff --git a/Engine/Source/Editor/MovieSceneTools/Private/TrackEditors/SkeletalAnimationTrackEditor.cpp b/Engine/Source/Editor/MovieSceneTools/Private/TrackEditors/SkeletalAnimationTrackEditor.cpp
--- a/Engine/Source/Editor/MovieSceneTools/Private/TrackEditors/SkeletalAnimationTrackEditor.cpp
+++ b/Engine/Source/Editor/MovieSceneTools/Private/TrackEditors/SkeletalAnimationTrackEditor.cpp
@@ -1048,7 +1048,7 @@ void FSkeletalAnimationTrackEditor::HandleCreatePoseAsset(FGuid InObjectBinding)
 	USkeleton* Skeleton = AcquireSkeletonFromObjectGuid(InObjectBinding, GetSequencer());
 	if (Skeleton)
 	{
-		TArray<TWeakObjectPtr<UObject>> Skeletons;
+		TArray<TSoftObjectPtr<UObject>> Skeletons;
 		Skeletons.Add(Skeleton);
 		AnimationEditorUtils::ExecuteNewAnimAsset<UPoseAssetFactory, UPoseAsset>(Skeletons, FString("_PoseAsset"), FAnimAssetCreated::CreateSP(this, &FSkeletalAnimationTrackEditor::CreatePoseAsset, InObjectBinding), false, false);
 	}
@@ -1533,7 +1533,7 @@ void FSkeletalAnimationTrackEditor::HandleCreateAnimationSequence(USkeletalMeshC
 {
 	if (SkelMeshComp)
 	{
-		TArray<TWeakObjectPtr<UObject>> Skels;
+		TArray<TSoftObjectPtr<UObject>> Skels;
 		if (SkelMeshComp->GetSkeletalMeshAsset())
 		{
 			Skels.Add(SkelMeshComp->GetSkeletalMeshAsset());
diff --git a/Engine/Source/Editor/Persona/Private/PersonaModule.cpp b/Engine/Source/Editor/Persona/Private/PersonaModule.cpp
--- a/Engine/Source/Editor/Persona/Private/PersonaModule.cpp
+++ b/Engine/Source/Editor/Persona/Private/PersonaModule.cpp
@@ -1255,7 +1255,7 @@ TSharedRef< SWidget > FPersonaModule::GenerateCreateAssetMenu(TWeakPtr<IPersonaT
 	MenuBuilder.EndSection();
 
 	TSharedRef<IPersonaToolkit> PersonaToolkit = InWeakPersonaToolkit.Pin().ToSharedRef();
-	TArray<TWeakObjectPtr<UObject>> Objects;
+	TArray<TSoftObjectPtr<UObject>> Objects;
 	if (PersonaToolkit->GetPreviewMesh())
 	{
 		Objects.Add(PersonaToolkit->GetPreviewMesh());
@@ -1273,7 +1273,7 @@ TSharedRef< SWidget > FPersonaModule::GenerateCreateAssetMenu(TWeakPtr<IPersonaT
 void FPersonaModule::FillCreateAnimationMenu(FMenuBuilder& MenuBuilder, TWeakPtr<IPersonaToolkit> InWeakPersonaToolkit)
 {
 	TSharedRef<IPersonaToolkit> PersonaToolkit = InWeakPersonaToolkit.Pin().ToSharedRef();
-	TArray<TWeakObjectPtr<UObject>> Objects;
+	TArray<TSoftObjectPtr<UObject>> Objects;
 	if (PersonaToolkit->GetPreviewMesh())
 	{
 		Objects.Add(PersonaToolkit->GetPreviewMesh());
@@ -1326,7 +1326,7 @@ void FPersonaModule::FillCreateAnimationMenu(FMenuBuilder& MenuBuilder, TWeakPtr
 void FPersonaModule::FillCreateAnimationFromCurrentAnimationMenu(FMenuBuilder& MenuBuilder, TWeakPtr<IPersonaToolkit> InWeakPersonaToolkit) 
 {
 	TSharedRef<IPersonaToolkit> PersonaToolkit = InWeakPersonaToolkit.Pin().ToSharedRef();
-	TArray<TWeakObjectPtr<UObject>> Objects;
+	TArray<TSoftObjectPtr<UObject>> Objects;
 
 	if (PersonaToolkit->GetPreviewMesh())
 	{
@@ -1367,7 +1367,7 @@ void FPersonaModule::FillCreateAnimationFromCurrentAnimationMenu(FMenuBuilder& M
 void FPersonaModule::FillCreatePoseAssetMenu(FMenuBuilder& MenuBuilder, TWeakPtr<IPersonaToolkit> InWeakPersonaToolkit)
 {
 	TSharedRef<IPersonaToolkit> PersonaToolkit = InWeakPersonaToolkit.Pin().ToSharedRef();
-	TArray<TWeakObjectPtr<UObject>> Objects;
+	TArray<TSoftObjectPtr<UObject>> Objects;
 
 	if (PersonaToolkit->GetPreviewMesh())
 	{
diff --git a/Engine/Source/Editor/StringTableEditor/Private/AssetDefinition_StringTable.cpp b/Engine/Source/Editor/StringTableEditor/Private/AssetDefinition_StringTable.cpp
new file mode 100644
--- /dev/null
+++ b/Engine/Source/Editor/StringTableEditor/Private/AssetDefinition_StringTable.cpp
@@ -0,0 +1,21 @@
+// Copyright Epic Games, Inc. All Rights Reserved.
+
+#include "AssetDefinition_StringTable.h"
+#include "StringTableEditorModule.h"
+#include "Modules/ModuleManager.h"
+
+#define LOCTEXT_NAMESPACE "AssetTypeActions"
+
+EAssetCommandResult UAssetDefinition_StringTable::OpenAssets(const FAssetOpenArgs& OpenArgs) const
+{
+	FStringTableEditorModule& StringTableEditorModule = FModuleManager::LoadModuleChecked<FStringTableEditorModule>("StringTableEditor");
+
+	for (UStringTable* StringTable : OpenArgs.LoadObjects<UStringTable>())
+	{
+		StringTableEditorModule.CreateStringTableEditor(OpenArgs.GetToolkitMode(), OpenArgs.ToolkitHost, StringTable);
+	}
+
+	return EAssetCommandResult::Handled;
+}
+
+#undef LOCTEXT_NAMESPACE
diff --git a/Engine/Source/Editor/StringTableEditor/Private/AssetDefinition_StringTable.h b/Engine/Source/Editor/StringTableEditor/Private/AssetDefinition_StringTable.h
new file mode 100644
--- /dev/null
+++ b/Engine/Source/Editor/StringTableEditor/Private/AssetDefinition_StringTable.h
@@ -0,0 +1,32 @@
+// Copyright Epic Games, Inc. All Rights Reserved.
+
+#pragma once
+
+#include "CoreMinimal.h"
+#include "Internationalization/StringTable.h"
+#include "AssetDefinitionDefault.h"
+
+#include "AssetDefinition_StringTable.generated.h"
+
+UCLASS()
+class UAssetDefinition_StringTable : public UAssetDefinitionDefault
+{
+	GENERATED_BODY()
+
+public:
+	// UAssetDefinition Begin
+	virtual FText GetAssetDisplayName() const override { return NSLOCTEXT("AssetTypeActions", "AssetTypeActions_StringTable", "String Table"); }
+	virtual FLinearColor GetAssetColor() const override { return FLinearColor(FColor(255, 196, 128)); }
+	virtual TSoftClassPtr<UObject> GetAssetClass() const override { return UStringTable::StaticClass(); }
+	virtual TConstArrayView<FAssetCategoryPath> GetAssetCategories() const override
+	{
+		static const auto Categories = { EAssetCategoryPaths::Misc };
+		return Categories;
+	}
+	virtual FAssetSupportResponse CanLocalize(const FAssetData& InAsset) const override
+	{
+		return FAssetSupportResponse::NotSupported();
+	}
+	virtual EAssetCommandResult OpenAssets(const FAssetOpenArgs& OpenArgs) const override;
+	// UAssetDefinition End
+};
diff --git a/Engine/Source/Editor/StringTableEditor/Private/AssetTypeActions_StringTable.cpp b/Engine/Source/Editor/StringTableEditor/Private/AssetTypeActions_StringTable.cpp
deleted file mode 100644
diff --git a/Engine/Source/Editor/StringTableEditor/Private/AssetTypeActions_StringTable.h b/Engine/Source/Editor/StringTableEditor/Private/AssetTypeActions_StringTable.h
deleted file mode 100644
diff --git a/Engine/Source/Editor/StringTableEditor/Private/StringTableEditorModule.cpp b/Engine/Source/Editor/StringTableEditor/Private/StringTableEditorModule.cpp
--- a/Engine/Source/Editor/StringTableEditor/Private/StringTableEditorModule.cpp
+++ b/Engine/Source/Editor/StringTableEditor/Private/StringTableEditorModule.cpp
@@ -1,12 +1,16 @@
 // Copyright Epic Games, Inc. All Rights Reserved.
 
 #include "StringTableEditorModule.h"
+
+#include "AssetToolsModule.h"
+#include "IAssetTools.h"
 #include "Modules/ModuleManager.h"
 #include "IStringTableEditor.h"
 #include "Internationalization/StringTableRegistry.h"
 #include "StringTableEditor.h"
-#include "AssetTypeActions_StringTable.h"
+
 #include "PackageMigrationContext.h"
+#include "Internationalization/StringTable.h"
 
 IMPLEMENT_MODULE(FStringTableEditorModule, StringTableEditor);
 
@@ -18,8 +22,6 @@ void FStringTableEditorModule::StartupModule()
 	ToolBarExtensibilityManager = MakeShareable(new FExtensibilityManager);
 
 	IAssetTools& AssetTools = FModuleManager::LoadModuleChecked<FAssetToolsModule>("AssetTools").Get();
-	AssetTools.RegisterAssetTypeActions(MakeShareable(new FAssetTypeActions_StringTable()));
-
 	AssetTools.GetOnPackageMigration().AddRaw(this, &FStringTableEditorModule::OnPackageMigration);
 }
 
diff --git a/Engine/Source/Editor/StringTableEditor/StringTableEditor.Build.cs b/Engine/Source/Editor/StringTableEditor/StringTableEditor.Build.cs
--- a/Engine/Source/Editor/StringTableEditor/StringTableEditor.Build.cs
+++ b/Engine/Source/Editor/StringTableEditor/StringTableEditor.Build.cs
@@ -19,6 +19,7 @@ public class StringTableEditor : ModuleRules
 				"EditorFramework",
                 "UnrealEd",
 				"AssetTools",
+				"AssetDefinition",
 			});
 
         DynamicallyLoadedModuleNames.Add("WorkspaceMenuStructure");
diff --git a/Engine/Source/Editor/UnrealEd/Private/AnimationEditorUtils.cpp b/Engine/Source/Editor/UnrealEd/Private/AnimationEditorUtils.cpp
--- a/Engine/Source/Editor/UnrealEd/Private/AnimationEditorUtils.cpp
+++ b/Engine/Source/Editor/UnrealEd/Private/AnimationEditorUtils.cpp
@@ -461,16 +461,16 @@ namespace AnimationEditorUtils
 		AssetToolsModule.Get().CreateUniqueAssetName(InBasePackageName, InSuffix, OutPackageName, OutAssetName);
 	}
 
-	void CreateAnimationAssets(const TArray<TWeakObjectPtr<UObject>>& SkeletonsOrSkeletalMeshes, TSubclassOf<UAnimationAsset> AssetClass, const FString& InPrefix, FAnimAssetCreated AssetCreated, UObject* NameBaseObject /*= nullptr*/, bool bDoNotShowNameDialog /*= false*/, bool bAllowReplaceExisting /*= false*/)
+	void CreateAnimationAssets(const TArray<TSoftObjectPtr<UObject>>& SkeletonsOrSkeletalMeshes, TSubclassOf<UAnimationAsset> AssetClass, const FString& InPrefix, FAnimAssetCreated AssetCreated, UObject* NameBaseObject /*= nullptr*/, bool bDoNotShowNameDialog /*= false*/, bool bAllowReplaceExisting /*= false*/)
 	{
 		TArray<UObject*> ObjectsToSync;
 		for(auto SkelIt = SkeletonsOrSkeletalMeshes.CreateConstIterator(); SkelIt; ++SkelIt)
 		{
 			USkeletalMesh* SkeletalMesh = nullptr;
-			USkeleton* Skeleton = Cast<USkeleton>(SkelIt->Get());
+			USkeleton* Skeleton = Cast<USkeleton>(SkelIt->LoadSynchronous());
 			if (Skeleton == nullptr)
 			{
-				SkeletalMesh = CastChecked<USkeletalMesh>(SkelIt->Get());
+				SkeletalMesh = CastChecked<USkeletalMesh>(SkelIt->LoadSynchronous());
 				Skeleton = SkeletalMesh->GetSkeleton();
 			}
 
@@ -565,17 +565,17 @@ namespace AnimationEditorUtils
 		}
 	}
 
-	void CreateNewAnimBlueprint(TArray<TWeakObjectPtr<UObject>> SkeletonsOrSkeletalMeshes, FAnimAssetCreated AssetCreated, bool bInContentBrowser)
+	void CreateNewAnimBlueprint(TArray<TSoftObjectPtr<UObject>> SkeletonsOrSkeletalMeshes, FAnimAssetCreated AssetCreated, bool bInContentBrowser)
 	{
 		const FString DefaultSuffix = TEXT("_AnimBlueprint");
 
 		if (SkeletonsOrSkeletalMeshes.Num() == 1)
 		{
 			USkeletalMesh* SkeletalMesh = nullptr;
-			USkeleton* Skeleton = Cast<USkeleton>(SkeletonsOrSkeletalMeshes[0].Get());
+			USkeleton* Skeleton = Cast<USkeleton>(SkeletonsOrSkeletalMeshes[0].LoadSynchronous());
 			if (Skeleton == nullptr)
 			{
-				SkeletalMesh = CastChecked<USkeletalMesh>(SkeletonsOrSkeletalMeshes[0].Get());
+				SkeletalMesh = CastChecked<USkeletalMesh>(SkeletonsOrSkeletalMeshes[0].LoadSynchronous());
 				Skeleton = SkeletalMesh->GetSkeleton();
 			}
 
@@ -625,10 +625,10 @@ namespace AnimationEditorUtils
 			for (auto ObjIt = SkeletonsOrSkeletalMeshes.CreateConstIterator(); ObjIt; ++ObjIt)
 			{
 				USkeletalMesh* SkeletalMesh = nullptr;
-				USkeleton* Skeleton = Cast<USkeleton>(ObjIt->Get());
+				USkeleton* Skeleton = Cast<USkeleton>(ObjIt->LoadSynchronous());
 				if (Skeleton == nullptr)
 				{
-					SkeletalMesh = CastChecked<USkeletalMesh>(ObjIt->Get());
+					SkeletalMesh = CastChecked<USkeletalMesh>(ObjIt->LoadSynchronous());
 					Skeleton = SkeletalMesh->GetSkeleton();
 				}
 
@@ -677,7 +677,7 @@ namespace AnimationEditorUtils
 		return AssetTools.IsAssetClassSupported(InClass);
 	}
 
-	void FillCreateAssetMenu(FMenuBuilder& MenuBuilder, const TArray<TWeakObjectPtr<UObject>>& SkeletonsOrSkeletalMeshes, FAnimAssetCreated AssetCreated, bool bInContentBrowser)
+	void FillCreateAssetMenu(FMenuBuilder& MenuBuilder, const TArray<TSoftObjectPtr<UObject>>& SkeletonsOrSkeletalMeshes, FAnimAssetCreated AssetCreated, bool bInContentBrowser)
 	{
 		const bool bAllowReplaceExisting = false;
 
diff --git a/Engine/Source/Editor/UnrealEd/Public/AnimationEditorUtils.h b/Engine/Source/Editor/UnrealEd/Public/AnimationEditorUtils.h
--- a/Engine/Source/Editor/UnrealEd/Public/AnimationEditorUtils.h
+++ b/Engine/Source/Editor/UnrealEd/Public/AnimationEditorUtils.h
@@ -17,6 +17,7 @@
 #include "IAssetTools.h"
 #include "AssetToolsModule.h"
 #include "Engine/SkeletalMesh.h"
+#include "UObject/SoftObjectPtr.h"
 
 class FMenuBuilder;
 class UAnimBlueprint;
@@ -138,10 +139,10 @@ namespace AnimationEditorUtils
 {
 	UNREALED_API FAssetData CreateModalAnimationCompressionSelectionDialog(const FAnimationCompressionSelectionDialogConfig& InConfig);
 
-	UNREALED_API void CreateAnimationAssets(const TArray<TWeakObjectPtr<UObject>>& SkeletonsOrSkeletalMeshes, TSubclassOf<UAnimationAsset> AssetClass, const FString& InPrefix, FAnimAssetCreated AssetCreated, UObject* NameBaseObject = nullptr, bool bDoNotShowNameDialog = false, bool bAllowReplaceExisting = false);
+	UNREALED_API void CreateAnimationAssets(const TArray<TSoftObjectPtr<UObject>>& SkeletonsOrSkeletalMeshes, TSubclassOf<UAnimationAsset> AssetClass, const FString& InPrefix, FAnimAssetCreated AssetCreated, UObject* NameBaseObject = nullptr, bool bDoNotShowNameDialog = false, bool bAllowReplaceExisting = false);
 	
-	UNREALED_API void CreateNewAnimBlueprint(TArray<TWeakObjectPtr<UObject>> SkeletonsOrSkeletalMeshes, FAnimAssetCreated AssetCreated, bool bInContentBrowser);
-	UNREALED_API void FillCreateAssetMenu(FMenuBuilder& MenuBuilder, const TArray<TWeakObjectPtr<UObject>>& SkeletonsOrSkeletalMeshes, FAnimAssetCreated AssetCreated, bool bInContentBrowser=true);
+	UNREALED_API void CreateNewAnimBlueprint(TArray<TSoftObjectPtr<UObject>> SkeletonsOrSkeletalMeshes, FAnimAssetCreated AssetCreated, bool bInContentBrowser);
+	UNREALED_API void FillCreateAssetMenu(FMenuBuilder& MenuBuilder, const TArray<TSoftObjectPtr<UObject>>& SkeletonsOrSkeletalMeshes, FAnimAssetCreated AssetCreated, bool bInContentBrowser=true);
 	UNREALED_API void CreateUniqueAssetName(const FString& InBasePackageName, const FString& InSuffix, FString& OutPackageName, FString& OutAssetName);
 
 	/** Applies the animation compression codecs to the sequence list with optional override settings */
@@ -212,15 +213,15 @@ namespace AnimationEditorUtils
 	//////////////////////////////////////////////////////////////////////////////////////////
 
 	template <typename TFactory, typename T>
-	void ExecuteNewAnimAsset(TArray<TWeakObjectPtr<UObject>> SkeletonsOrSkeletalMeshes, const FString InSuffix, FAnimAssetCreated AssetCreated, bool bInContentBrowser, bool bAllowReplaceExisting)
+	void ExecuteNewAnimAsset(TArray<TSoftObjectPtr<UObject>> SkeletonsOrSkeletalMeshes, const FString InSuffix, FAnimAssetCreated AssetCreated, bool bInContentBrowser, bool bAllowReplaceExisting)
 	{
 		if(bInContentBrowser && SkeletonsOrSkeletalMeshes.Num() == 1)
 		{
 			USkeletalMesh* SkeletalMesh = nullptr;
-			USkeleton* Skeleton = Cast<USkeleton>(SkeletonsOrSkeletalMeshes[0].Get());
+			USkeleton* Skeleton = Cast<USkeleton>(SkeletonsOrSkeletalMeshes[0].LoadSynchronous());
 			if (Skeleton == nullptr)
 			{
-				SkeletalMesh = CastChecked<USkeletalMesh>(SkeletonsOrSkeletalMeshes[0].Get());
+				SkeletalMesh = CastChecked<USkeletalMesh>(SkeletonsOrSkeletalMeshes[0].LoadSynchronous());
 				Skeleton = SkeletalMesh->GetSkeleton();
 			}
 
