AssetDefinition - Upgrading the Data/Curve tables to use the new AssetDefinition.
AssetDefinition - Changing how getting the SourceFiles works a bit, wanted it to be easier to just get the real data and not have to wrap it myself, and also to not need to needlessly put it into a container if it's not needed.

diff --git a/Engine/Plugins/Editor/ContentBrowser/ContentBrowserAssetDataSource/Source/ContentBrowserAssetDataSource/Private/AssetFileContextMenu.cpp b/Engine/Plugins/Editor/ContentBrowser/ContentBrowserAssetDataSource/Source/ContentBrowserAssetDataSource/Private/AssetFileContextMenu.cpp
--- a/Engine/Plugins/Editor/ContentBrowser/ContentBrowserAssetDataSource/Source/ContentBrowserAssetDataSource/Private/AssetFileContextMenu.cpp
+++ b/Engine/Plugins/Editor/ContentBrowser/ContentBrowserAssetDataSource/Source/ContentBrowserAssetDataSource/Private/AssetFileContextMenu.cpp
@@ -35,6 +35,7 @@
 
 #include "ObjectTools.h"
 #include "Editor.h"
+#include "EditorFramework/AssetImportData.h"
 
 #include "PropertyEditorModule.h"
 #include "ConsolidateWindow.h"
@@ -1218,15 +1219,15 @@ void FAssetFileContextMenu::ExecuteReimportWithNewFile(int32 SourceFileIndex /*=
 
 	if (const UAssetDefinition* AssetDefinition = UAssetDefinitionRegistry::Get()->GetAssetDefinitionForAsset(SelectedAssets[0]))
 	{
-		TArray<FAssetSourceFile> OutSourceAssets;
-		FAssetSourceFileArgs SourceFileArgs(SelectedAssets);
-		if (AssetDefinition->GetSourceFiles(SourceFileArgs, OutSourceAssets) == EAssetCommandResult::Handled)
+		FAssetData SelectedAsset = SelectedAssets[0];
+		
+		AssetDefinition->GetSourceFiles(SelectedAsset, [&](const FAssetImportInfo& AssetImportInfo)
 		{
 			int32 SourceFileIndexToReplace = SourceFileIndex;
 			//Check if the data is valid
 			if (SourceFileIndex == INDEX_NONE)
 			{
-				if (OutSourceAssets.Num() > 1)
+				if (AssetImportInfo.SourceFiles.Num() > 1)
 				{
 					//Ask for a new file for the index 0
 					SourceFileIndexToReplace = 0;
@@ -1234,12 +1235,12 @@ void FAssetFileContextMenu::ExecuteReimportWithNewFile(int32 SourceFileIndex /*=
 			}
 			else
 			{
-				check(OutSourceAssets.IsValidIndex(SourceFileIndex));
+				check(AssetImportInfo.SourceFiles.IsValidIndex(SourceFileIndex));
 			}
-			
-			TArray<UObject*> LoadedAssets = SourceFileArgs.LoadObjects<UObject>();
+
+			TArray<UObject*> LoadedAssets = { SelectedAsset.GetAsset() };
 			FReimportManager::Instance()->ValidateAllSourceFileAndReimport(LoadedAssets, true, SourceFileIndexToReplace, true);
-		}
+		});
 	}
 }
 
@@ -1287,16 +1288,17 @@ void FAssetFileContextMenu::GetSelectedAssetSourceFilePaths(TArray<FString>& Out
 	{
 		if (const UAssetDefinition* AssetDefinition = UAssetDefinitionRegistry::Get()->GetAssetDefinitionForClass(AssetsByClassPair.Key))
 		{
-			TArray<FAssetSourceFile> OutSourceAssets;
-			FAssetSourceFileArgs SourceFileArgs(SelectedAssets);
-			if (AssetDefinition->GetSourceFiles(SourceFileArgs, OutSourceAssets) == EAssetCommandResult::Handled)
+			for (const FAssetData& Asset : AssetsByClassPair.Value)
 			{
-				OutValidSelectedAssetCount += AssetsByClassPair.Value.Num();
-				for (const FAssetSourceFile& SourceFile : OutSourceAssets)
+				AssetDefinition->GetSourceFiles(Asset, [&](const FAssetImportInfo& AssetImportInfo)
 				{
-					OutFilePaths.Add(SourceFile.RelativeFilename);
-					OutUniqueSourceFileLabels.AddUnique(SourceFile.DisplayLabelName);
-				}
+					OutValidSelectedAssetCount++;
+					for (const auto& SourceFile : AssetImportInfo.SourceFiles)
+					{
+						OutFilePaths.Add(SourceFile.RelativeFilename);
+						OutUniqueSourceFileLabels.AddUnique(SourceFile.DisplayLabelName);
+					}
+				});
 			}
 		}
 	}
diff --git a/Engine/Plugins/Editor/EngineAssetDefinitions/Source/EngineAssetDefinitions.Build.cs b/Engine/Plugins/Editor/EngineAssetDefinitions/Source/EngineAssetDefinitions.Build.cs
--- a/Engine/Plugins/Editor/EngineAssetDefinitions/Source/EngineAssetDefinitions.Build.cs
+++ b/Engine/Plugins/Editor/EngineAssetDefinitions/Source/EngineAssetDefinitions.Build.cs
@@ -44,7 +44,9 @@ public class EngineAssetDefinitions : ModuleRules
 				"TextureEditor",
 				"MaterialEditor",
 				"PhysicsAssetEditor",
-				"FontEditor"
+				"FontEditor",
+				"DataTableEditor",
+				"CurveTableEditor",
 			}
 		);
 
@@ -56,7 +58,9 @@ public class EngineAssetDefinitions : ModuleRules
 				"TextureEditor",
 				"MaterialEditor",
 				"PhysicsAssetEditor",
-				"FontEditor"
+				"FontEditor",
+				"DataTableEditor",
+				"CurveTableEditor",
 			}
 		);
 	}
diff --git a/Engine/Plugins/Editor/EngineAssetDefinitions/Source/Private/Table/AssetDefinition_CompositeCurveTable.h b/Engine/Plugins/Editor/EngineAssetDefinitions/Source/Private/Table/AssetDefinition_CompositeCurveTable.h
new file mode 100644
--- /dev/null
+++ b/Engine/Plugins/Editor/EngineAssetDefinitions/Source/Private/Table/AssetDefinition_CompositeCurveTable.h
@@ -0,0 +1,21 @@
+// Copyright Epic Games, Inc. All Rights Reserved.
+
+#pragma once
+
+#include "Engine/CompositeCurveTable.h"
+#include "AssetDefinition_CurveTable.h"
+
+#include "AssetDefinition_CompositeCurveTable.generated.h"
+
+UCLASS()
+class UAssetDefinition_CompositeCurveTable : public UAssetDefinition_CurveTable
+{
+	GENERATED_BODY()
+
+public:
+	// UAssetDefinition Begin
+	virtual FText GetAssetDisplayName() const override { return NSLOCTEXT("AssetTypeActions", "AssetTypeActions_CompositeCurveTable", "Composite Curve Table"); }
+	virtual TSoftClassPtr<UObject> GetAssetClass() const override { return UCompositeCurveTable::StaticClass(); }
+	virtual bool CanImport() const override { return false; }
+	// UAssetDefinition End
+};
diff --git a/Engine/Plugins/Editor/EngineAssetDefinitions/Source/Private/Table/AssetDefinition_CompositeDataTable.h b/Engine/Plugins/Editor/EngineAssetDefinitions/Source/Private/Table/AssetDefinition_CompositeDataTable.h
new file mode 100644
--- /dev/null
+++ b/Engine/Plugins/Editor/EngineAssetDefinitions/Source/Private/Table/AssetDefinition_CompositeDataTable.h
@@ -0,0 +1,21 @@
+// Copyright Epic Games, Inc. All Rights Reserved.
+
+#pragma once
+
+#include "Engine/CompositeDataTable.h"
+#include "AssetDefinition_DataTable.h"
+
+#include "AssetDefinition_CompositeDataTable.generated.h"
+
+UCLASS()
+class UAssetDefinition_CompositeDataTable : public UAssetDefinition_DataTable
+{
+	GENERATED_BODY()
+
+public:
+	// UAssetDefinition Begin
+	virtual FText GetAssetDisplayName() const override { return NSLOCTEXT("AssetTypeActions", "FAssetTypeActions_CompositeDataTable", "Composite Data Table"); }
+	virtual TSoftClassPtr<UObject> GetAssetClass() const override { return UCompositeDataTable::StaticClass(); }
+	virtual bool CanImport() const override { return false; }
+	// UAssetDefinition End
+};
diff --git a/Engine/Plugins/Editor/EngineAssetDefinitions/Source/Private/Table/AssetDefinition_CurveTable.cpp b/Engine/Plugins/Editor/EngineAssetDefinitions/Source/Private/Table/AssetDefinition_CurveTable.cpp
new file mode 100644
--- /dev/null
+++ b/Engine/Plugins/Editor/EngineAssetDefinitions/Source/Private/Table/AssetDefinition_CurveTable.cpp
@@ -0,0 +1,221 @@
+// Copyright Epic Games, Inc. All Rights Reserved.
+
+#include "AssetDefinition_CurveTable.h"
+
+#include "AssetDefinitionRegistry.h"
+#include "ToolMenus.h"
+#include "Misc/FileHelper.h"
+#include "EditorFramework/AssetImportData.h"
+#include "Framework/Application/SlateApplication.h"
+#include "AssetToolsModule.h"
+#include "ContentBrowserMenuContexts.h"
+#include "CurveTableEditorModule.h"
+#include "DesktopPlatformModule.h"
+#include "FindSourceFileInExplorer.h"
+#include "ToolMenuDelegates.h"
+#include "Settings/EditorLoadingSavingSettings.h"
+
+#define LOCTEXT_NAMESPACE "AssetTypeActions"
+
+EAssetCommandResult UAssetDefinition_CurveTable::PerformAssetDiff(const FAssetDiffArgs& DiffArgs) const
+{
+	UCurveTable* OldCurveTable = CastChecked<UCurveTable>(DiffArgs.OldAsset);
+	UCurveTable* NewCurveTable = CastChecked<UCurveTable>(DiffArgs.NewAsset);
+
+	// Build names for temp csv files
+	const FString RelOldTempFileName = FString::Printf(TEXT("%sTemp%s-%s.csv"), *FPaths::DiffDir(), *DiffArgs.OldAsset->GetName(), *DiffArgs.OldRevision.Revision);
+	const FString AbsoluteOldTempFileName = FPaths::ConvertRelativePathToFull(RelOldTempFileName);
+	const FString RelNewTempFileName = FString::Printf(TEXT("%sTemp%s-%s.csv"), *FPaths::DiffDir(), *DiffArgs.NewAsset->GetName(), *DiffArgs.NewRevision.Revision);
+	const FString AbsoluteNewTempFileName = FPaths::ConvertRelativePathToFull(RelNewTempFileName);
+
+	// save temp files
+	const bool OldResult = FFileHelper::SaveStringToFile(OldCurveTable->GetTableAsCSV(), *AbsoluteOldTempFileName);
+	const bool NewResult = FFileHelper::SaveStringToFile(NewCurveTable->GetTableAsCSV(), *AbsoluteNewTempFileName);
+
+	if (OldResult && NewResult)
+	{
+		const FString DiffCommand = GetDefault<UEditorLoadingSavingSettings>()->TextDiffToolPath.FilePath;
+
+		FAssetToolsModule& AssetToolsModule = FModuleManager::Get().LoadModuleChecked<FAssetToolsModule>("AssetTools");
+		AssetToolsModule.Get().CreateDiffProcess(DiffCommand, AbsoluteOldTempFileName, AbsoluteNewTempFileName);
+
+		return EAssetCommandResult::Handled;
+	}
+
+	return Super::PerformAssetDiff(DiffArgs);
+}
+
+EAssetCommandResult UAssetDefinition_CurveTable::OpenAssets(const FAssetOpenArgs& OpenArgs) const
+{
+	FCurveTableEditorModule& CurveTableEditorModule = FModuleManager::LoadModuleChecked<FCurveTableEditorModule>( "CurveTableEditor" );
+	
+	for (UCurveTable* Table : OpenArgs.LoadObjects<UCurveTable>())
+	{
+		TSharedRef<ICurveTableEditor> NewCurveTableEditor = CurveTableEditorModule.CreateCurveTableEditor(OpenArgs.GetToolkitMode(), OpenArgs.ToolkitHost, Table);
+	}
+	
+	return EAssetCommandResult::Handled;
+}
+
+// Menu Extensions
+//--------------------------------------------------------------------
+
+namespace MenuExtension_CurveTable
+{
+	void ExecuteExportAsCSV(const FToolMenuContext& InContext)
+	{
+		IDesktopPlatform* DesktopPlatform = FDesktopPlatformModule::Get();
+
+		const void* ParentWindowWindowHandle = FSlateApplication::Get().FindBestParentWindowHandleForDialogs(nullptr);
+
+		const UContentBrowserAssetContextMenuContext* Context = UContentBrowserAssetContextMenuContext::FindContextWithAssets(InContext);
+		for (UCurveTable* CurveTable : Context->LoadSelectedObjects<UCurveTable>())
+		{
+			const FText Title = FText::Format(LOCTEXT("CurveTable_ExportCSVDialogTitle", "Export '{0}' as CSV..."), FText::FromString(*CurveTable->GetName()));
+			const FString CurrentFilename = CurveTable->AssetImportData->GetFirstFilename();
+			const FString FileTypes = TEXT("Data Table CSV (*.csv)|*.csv");
+
+			TArray<FString> OutFilenames;
+			DesktopPlatform->SaveFileDialog(
+				ParentWindowWindowHandle,
+				Title.ToString(),
+				(CurrentFilename.IsEmpty()) ? TEXT("") : FPaths::GetPath(CurrentFilename),
+				(CurrentFilename.IsEmpty()) ? TEXT("") : FPaths::GetBaseFilename(CurrentFilename) + TEXT(".csv"),
+				FileTypes,
+				EFileDialogFlags::None,
+				OutFilenames
+				);
+
+			if (OutFilenames.Num() > 0)
+			{
+				FFileHelper::SaveStringToFile(CurveTable->GetTableAsCSV(), *OutFilenames[0]);
+			}
+		}
+	}
+
+	void ExecuteExportAsJSON(const FToolMenuContext& InContext)
+	{
+		IDesktopPlatform* DesktopPlatform = FDesktopPlatformModule::Get();
+
+		const void* ParentWindowWindowHandle = FSlateApplication::Get().FindBestParentWindowHandleForDialogs(nullptr);
+
+		const UContentBrowserAssetContextMenuContext* Context = UContentBrowserAssetContextMenuContext::FindContextWithAssets(InContext);
+		for (UCurveTable* CurveTable : Context->LoadSelectedObjects<UCurveTable>())
+		{
+			const FText Title = FText::Format(LOCTEXT("DataTable_ExportJSONDialogTitle", "Export '{0}' as JSON..."), FText::FromString(*CurveTable->GetName()));
+			const FString CurrentFilename = CurveTable->AssetImportData->GetFirstFilename();
+			const FString FileTypes = TEXT("Data Table JSON (*.json)|*.json");
+
+			TArray<FString> OutFilenames;
+			DesktopPlatform->SaveFileDialog(
+				ParentWindowWindowHandle,
+				Title.ToString(),
+				(CurrentFilename.IsEmpty()) ? TEXT("") : FPaths::GetPath(CurrentFilename),
+				(CurrentFilename.IsEmpty()) ? TEXT("") : FPaths::GetBaseFilename(CurrentFilename) + TEXT(".json"),
+				FileTypes,
+				EFileDialogFlags::None,
+				OutFilenames
+			);
+
+			if (OutFilenames.Num() > 0)
+			{
+				FFileHelper::SaveStringToFile(CurveTable->GetTableAsJSON(), *OutFilenames[0]);
+			}
+		}
+	}
+
+	void ExecuteFindSourceFileInExplorer(const FToolMenuContext& InContext)
+	{
+		const UContentBrowserAssetContextMenuContext* Context = UContentBrowserAssetContextMenuContext::FindContextWithAssets(InContext);
+		
+		TArray<FString> ImportPaths;
+		for (const FAssetData& Asset : Context->SelectedAssets)
+		{
+			UAssetDefinitionRegistry::Get()->GetAssetDefinitionForAsset(Asset)->GetSourceFiles(Asset, [&ImportPaths](const FAssetImportInfo& AssetImportInfo)
+			{
+				for (const auto& SourceFile : AssetImportInfo.SourceFiles)
+				{
+					ImportPaths.Add(SourceFile.RelativeFilename);
+				}
+			});
+		}
+
+		TArray<FString> PotentialFileExtensions;
+		PotentialFileExtensions.Add(TEXT(".xls"));
+		PotentialFileExtensions.Add(TEXT(".xlsm"));
+		PotentialFileExtensions.Add(TEXT(".csv"));
+		PotentialFileExtensions.Add(TEXT(".json"));
+
+		UE::AssetTools::ExecuteFindSourceFileInExplorer(ImportPaths, PotentialFileExtensions);
+	}
+
+	bool CanExecuteFindSourceFileInExplorer(const FToolMenuContext& InContext)
+	{
+		const UContentBrowserAssetContextMenuContext* Context = UContentBrowserAssetContextMenuContext::FindContextWithAssets(InContext);
+
+		TArray<FString> ImportPaths;
+		for (const FAssetData& Asset : Context->SelectedAssets)
+		{
+			UAssetDefinitionRegistry::Get()->GetAssetDefinitionForAsset(Asset)->GetSourceFiles(Asset, [&ImportPaths](const FAssetImportInfo& AssetImportInfo)
+			{
+				for (const auto& SourceFile : AssetImportInfo.SourceFiles)
+				{
+					ImportPaths.Add(SourceFile.RelativeFilename);
+				}
+			});
+		}
+
+		TArray<FString> PotentialFileExtensions;
+		PotentialFileExtensions.Add(TEXT(".xls"));
+		PotentialFileExtensions.Add(TEXT(".xlsm"));
+		PotentialFileExtensions.Add(TEXT(".csv"));
+		PotentialFileExtensions.Add(TEXT(".json"));
+
+		return UE::AssetTools::CanExecuteFindSourceFileInExplorer(ImportPaths, PotentialFileExtensions);
+	}
+	
+	static FDelayedAutoRegisterHelper DelayedAutoRegister(EDelayedRegisterRunPhase::EndOfEngineInit, []{ 
+		UToolMenus::RegisterStartupCallback(FSimpleMulticastDelegate::FDelegate::CreateLambda([]()
+		{
+			FToolMenuOwnerScoped OwnerScoped(UE_MODULE_NAME);
+			UToolMenu* Menu = UE::ContentBrowser::ExtendToolMenu_AssetContextMenu(UCurveTable::StaticClass());
+		
+			FToolMenuSection& Section = Menu->FindOrAddSection("GetAssetActions");
+			Section.AddDynamicEntry(NAME_None, FNewToolMenuSectionDelegate::CreateLambda([](FToolMenuSection& InSection)
+			{
+				{
+					const TAttribute<FText> Label = LOCTEXT("CurveTable_ExportAsCSV", "Export as CSV");
+					const TAttribute<FText> ToolTip = LOCTEXT("CurveTable_ExportAsCSVTooltip", "Export the curve table as a file containing CSV data.");
+					const FSlateIcon Icon = FSlateIcon();
+
+					FToolUIAction UIAction;
+					UIAction.ExecuteAction = FToolMenuExecuteAction::CreateStatic(&ExecuteExportAsCSV);
+					InSection.AddMenuEntry("CurveTable_ExportAsCSV", Label, ToolTip, Icon, UIAction);
+				}
+				
+				{
+					const TAttribute<FText> Label = LOCTEXT("CurveTable_ExportAsJSON", "Export as JSON");
+					const TAttribute<FText> ToolTip = LOCTEXT("CurveTable_ExportAsJSONTooltip", "Export the curve table as a file containing JSON data.");
+					const FSlateIcon Icon = FSlateIcon();
+
+					FToolUIAction UIAction;
+					UIAction.ExecuteAction = FToolMenuExecuteAction::CreateStatic(&ExecuteExportAsJSON);
+					InSection.AddMenuEntry("CurveTable_ExportAsJSON", Label, ToolTip, Icon, UIAction);
+				}
+
+				{
+					const TAttribute<FText> Label = LOCTEXT("CurveTable_OpenSourceData", "Open Source Data");
+					const TAttribute<FText> ToolTip = LOCTEXT("CurveTable_OpenSourceDataTooltip", "Opens the curve table's source data file in an external editor. It will search using the following extensions: .xls/.xlsm/.csv/.json");
+					const FSlateIcon Icon = FSlateIcon();
+
+					FToolUIAction UIAction;
+					UIAction.ExecuteAction = FToolMenuExecuteAction::CreateLambda(&ExecuteFindSourceFileInExplorer);
+					UIAction.CanExecuteAction = FToolMenuCanExecuteAction::CreateStatic(&CanExecuteFindSourceFileInExplorer);
+					InSection.AddMenuEntry("CurveTable_OpenSourceData", Label, ToolTip, Icon, UIAction);
+				}
+			}));
+		}));
+	});
+}
+
+#undef LOCTEXT_NAMESPACE
diff --git a/Engine/Plugins/Editor/EngineAssetDefinitions/Source/Private/Table/AssetDefinition_CurveTable.h b/Engine/Plugins/Editor/EngineAssetDefinitions/Source/Private/Table/AssetDefinition_CurveTable.h
new file mode 100644
--- /dev/null
+++ b/Engine/Plugins/Editor/EngineAssetDefinitions/Source/Private/Table/AssetDefinition_CurveTable.h
@@ -0,0 +1,30 @@
+// Copyright Epic Games, Inc. All Rights Reserved.
+
+#pragma once
+
+#include "CoreMinimal.h"
+#include "Engine/CurveTable.h"
+#include "AssetDefinitionDefault.h"
+
+#include "AssetDefinition_CurveTable.generated.h"
+
+UCLASS()
+class UAssetDefinition_CurveTable : public UAssetDefinitionDefault
+{
+	GENERATED_BODY()
+
+public:
+	// UAssetDefinition Begin
+	virtual FText GetAssetDisplayName() const override { return NSLOCTEXT("AssetTypeActions", "AssetTypeActions_CurveTable", "Curve Table"); }
+	virtual FLinearColor GetAssetColor() const override { return FLinearColor(FColor(62, 140, 35)); }
+	virtual TSoftClassPtr<UObject> GetAssetClass() const override { return UCurveTable::StaticClass(); }
+	virtual TConstArrayView<FAssetCategoryPath> GetAssetCategories() const override
+	{
+		static const auto Categories = { EAssetCategoryPaths::Misc };
+		return Categories;
+	}
+	virtual bool CanImport() const override { return true; }
+	virtual EAssetCommandResult PerformAssetDiff(const FAssetDiffArgs& DiffArgs) const override;
+	virtual EAssetCommandResult OpenAssets(const FAssetOpenArgs& OpenArgs) const override;
+	// UAssetDefinition End
+};
diff --git a/Engine/Plugins/Editor/EngineAssetDefinitions/Source/Private/Table/AssetDefinition_DataTable.cpp b/Engine/Plugins/Editor/EngineAssetDefinitions/Source/Private/Table/AssetDefinition_DataTable.cpp
new file mode 100644
--- /dev/null
+++ b/Engine/Plugins/Editor/EngineAssetDefinitions/Source/Private/Table/AssetDefinition_DataTable.cpp
@@ -0,0 +1,323 @@
+// Copyright Epic Games, Inc. All Rights Reserved.
+
+#include "AssetDefinition_DataTable.h"
+
+#include "AssetDefinitionRegistry.h"
+#include "ToolMenus.h"
+#include "Misc/FileHelper.h"
+#include "EditorFramework/AssetImportData.h"
+#include "Misc/MessageDialog.h"
+#include "Framework/Application/SlateApplication.h"
+#include "Settings/EditorLoadingSavingSettings.h"
+
+#include "DataTableEditorModule.h"
+#include "DesktopPlatformModule.h"
+#include "AssetToolsModule.h"
+#include "ToolMenu.h"
+#include "ToolMenus.h"
+#include "ContentBrowserMenuContexts.h"
+#include "FindSourceFileInExplorer.h"
+#include "ToolMenuSection.h"
+
+#define LOCTEXT_NAMESPACE "AssetTypeActions"
+
+static const FName NAME_RowStructure(TEXT("RowStructure"));
+
+FText UAssetDefinition_DataTable::GetAssetDisplayName(const FAssetData& AssetData) const
+{
+	if (AssetData.IsValid())
+	{
+		const FAssetDataTagMapSharedView::FFindTagResult RowStructureTag = AssetData.TagsAndValues.FindTag(NAME_RowStructure);
+		if (RowStructureTag.IsSet())
+		{
+			if (UScriptStruct* FoundStruct = UClass::TryFindTypeSlow<UScriptStruct>(RowStructureTag.GetValue(), EFindFirstObjectOptions::ExactClass))
+			{
+				return FText::Format(LOCTEXT("DataTableWithRowType", "Data Table ({0})"), FoundStruct->GetDisplayNameText());
+			}
+		}
+	}
+
+	return FText::GetEmpty();
+}
+
+// Attempts to export temporary CSV files and diff those. If that fails we fall back to diffing the data table assets directly.
+EAssetCommandResult UAssetDefinition_DataTable::PerformAssetDiff(const FAssetDiffArgs& DiffArgs) const
+{
+	const UDataTable* OldDataTable = CastChecked<UDataTable>(DiffArgs.OldAsset);
+	const UDataTable* NewDataTable = CastChecked<UDataTable>(DiffArgs.NewAsset);
+
+	// Build names for temp csv files
+	const FString RelOldTempFileName = FString::Printf(TEXT("%sTemp%s-%s.csv"), *FPaths::DiffDir(), *DiffArgs.OldAsset->GetName(), *DiffArgs.OldRevision.Revision);
+	const FString AbsoluteOldTempFileName = FPaths::ConvertRelativePathToFull(RelOldTempFileName);
+	const FString RelNewTempFileName = FString::Printf(TEXT("%sTemp%s-%s.csv"), *FPaths::DiffDir(), *DiffArgs.NewAsset->GetName(), *DiffArgs.NewRevision.Revision);
+	const FString AbsoluteNewTempFileName = FPaths::ConvertRelativePathToFull(RelNewTempFileName);
+
+	// save temp files
+	const bool OldResult = FFileHelper::SaveStringToFile(OldDataTable->GetTableAsCSV(EDataTableExportFlags::UseSimpleText), *AbsoluteOldTempFileName);
+	const bool NewResult = FFileHelper::SaveStringToFile(NewDataTable->GetTableAsCSV(EDataTableExportFlags::UseSimpleText), *AbsoluteNewTempFileName);
+
+	if (OldResult && NewResult)
+	{
+		const FString DiffCommand = GetDefault<UEditorLoadingSavingSettings>()->TextDiffToolPath.FilePath;
+
+		FAssetToolsModule& AssetToolsModule = FModuleManager::Get().LoadModuleChecked<FAssetToolsModule>("AssetTools");
+		AssetToolsModule.Get().CreateDiffProcess(DiffCommand, AbsoluteOldTempFileName, AbsoluteNewTempFileName);
+
+		return EAssetCommandResult::Handled;
+	}
+
+	return Super::PerformAssetDiff(DiffArgs);
+}
+
+EAssetCommandResult UAssetDefinition_DataTable::OpenAssets(const FAssetOpenArgs& OpenArgs) const
+{
+	TArray<UDataTable*> DataTablesToOpen;
+	TArray<UDataTable*> InvalidDataTables;
+
+	for (UDataTable* Table : OpenArgs.LoadObjects<UDataTable>())
+	{
+		if (Table->GetRowStruct())
+		{
+			DataTablesToOpen.Add(Table);
+		}
+		else
+		{
+			InvalidDataTables.Add(Table);
+		}
+	}
+
+	if (InvalidDataTables.Num() > 0)
+	{
+		FTextBuilder DataTablesListText;
+		DataTablesListText.Indent();
+		for (UDataTable* Table : InvalidDataTables)
+		{
+			const FTopLevelAssetPath ResolvedRowStructName = Table->GetRowStructPathName();
+			DataTablesListText.AppendLineFormat(LOCTEXT("DataTable_MissingRowStructListEntry", "* {0} (Row Structure: {1})"), FText::FromString(Table->GetName()), FText::FromString(ResolvedRowStructName.ToString()));
+		}
+
+		FText Title = LOCTEXT("DataTable_MissingRowStructTitle", "Continue?");
+		const EAppReturnType::Type DlgResult = FMessageDialog::Open(
+			EAppMsgType::YesNoCancel, 
+			FText::Format(LOCTEXT("DataTable_MissingRowStructMsg", "The following Data Tables are missing their row structure and will not be editable.\n\n{0}\n\nDo you want to open these data tables?"), DataTablesListText.ToText()), 
+			&Title
+			);
+
+		switch(DlgResult)
+		{
+		case EAppReturnType::Yes:
+			DataTablesToOpen.Append(InvalidDataTables);
+			break;
+		case EAppReturnType::Cancel:
+			return EAssetCommandResult::Handled;
+		default:
+			break;
+		}
+	}
+
+	FDataTableEditorModule& DataTableEditorModule = FModuleManager::LoadModuleChecked<FDataTableEditorModule>("DataTableEditor");
+	for (UDataTable* Table : DataTablesToOpen)
+	{
+		DataTableEditorModule.CreateDataTableEditor(OpenArgs.GetToolkitMode(), OpenArgs.ToolkitHost, Table);
+	}
+
+	return EAssetCommandResult::Handled;
+}
+
+
+// Menu Extensions
+//--------------------------------------------------------------------
+
+namespace MenuExtension_DataTable
+{
+	void ExecuteExportAsCSV(const FToolMenuContext& InContext)
+	{
+		IDesktopPlatform* DesktopPlatform = FDesktopPlatformModule::Get();
+
+		const void* ParentWindowWindowHandle = FSlateApplication::Get().FindBestParentWindowHandleForDialogs(nullptr);
+
+		const UContentBrowserAssetContextMenuContext* Context = UContentBrowserAssetContextMenuContext::FindContextWithAssets(InContext);
+		for (UDataTable* DataTable : Context->LoadSelectedObjects<UDataTable>())
+		{
+			const FText Title = FText::Format(LOCTEXT("DataTable_ExportCSVDialogTitle", "Export '{0}' as CSV..."), FText::FromString(*DataTable->GetName()));
+			const FString CurrentFilename = DataTable->AssetImportData->GetFirstFilename();
+			const FString FileTypes = TEXT("Data Table CSV (*.csv)|*.csv");
+
+			TArray<FString> OutFilenames;
+			DesktopPlatform->SaveFileDialog(
+				ParentWindowWindowHandle,
+				Title.ToString(),
+				(CurrentFilename.IsEmpty()) ? TEXT("") : FPaths::GetPath(CurrentFilename),
+				(CurrentFilename.IsEmpty()) ? TEXT("") : FPaths::GetBaseFilename(CurrentFilename) + TEXT(".csv"),
+				FileTypes,
+				EFileDialogFlags::None,
+				OutFilenames
+				);
+
+			if (OutFilenames.Num() > 0)
+			{
+				FFileHelper::SaveStringToFile(DataTable->GetTableAsCSV(), *OutFilenames[0]);
+			}
+		}
+	}
+
+	void ExecuteExportAsJSON(const FToolMenuContext& InContext)
+	{
+		IDesktopPlatform* DesktopPlatform = FDesktopPlatformModule::Get();
+
+		const void* ParentWindowWindowHandle = FSlateApplication::Get().FindBestParentWindowHandleForDialogs(nullptr);
+
+		const UContentBrowserAssetContextMenuContext* Context = UContentBrowserAssetContextMenuContext::FindContextWithAssets(InContext);
+		for (UDataTable* DataTable : Context->LoadSelectedObjects<UDataTable>())
+		{
+			const FText Title = FText::Format(LOCTEXT("DataTable_ExportJSONDialogTitle", "Export '{0}' as JSON..."), FText::FromString(*DataTable->GetName()));
+			const FString CurrentFilename = DataTable->AssetImportData->GetFirstFilename();
+			const FString FileTypes = TEXT("Data Table JSON (*.json)|*.json");
+
+			TArray<FString> OutFilenames;
+			DesktopPlatform->SaveFileDialog(
+				ParentWindowWindowHandle,
+				Title.ToString(),
+				(CurrentFilename.IsEmpty()) ? TEXT("") : FPaths::GetPath(CurrentFilename),
+				(CurrentFilename.IsEmpty()) ? TEXT("") : FPaths::GetBaseFilename(CurrentFilename) + TEXT(".json"),
+				FileTypes,
+				EFileDialogFlags::None,
+				OutFilenames
+			);
+
+			if (OutFilenames.Num() > 0)
+			{
+				FFileHelper::SaveStringToFile(DataTable->GetTableAsJSON(EDataTableExportFlags::UseJsonObjectsForStructs), *OutFilenames[0]);
+			}
+		}
+	}
+
+	void ExecuteFindSourceFileInExplorer(const FToolMenuContext& InContext)
+	{
+		const UContentBrowserAssetContextMenuContext* Context = UContentBrowserAssetContextMenuContext::FindContextWithAssets(InContext);
+		
+		TArray<FString> ImportPaths;
+		for (const FAssetData& Asset : Context->SelectedAssets)
+		{
+			UAssetDefinitionRegistry::Get()->GetAssetDefinitionForAsset(Asset)->GetSourceFiles(Asset, [&ImportPaths](const FAssetImportInfo& AssetImportInfo)
+			{
+				for (const auto& SourceFile : AssetImportInfo.SourceFiles)
+				{
+					ImportPaths.Add(SourceFile.RelativeFilename);
+				}
+			});
+		}
+
+		TArray<FString> PotentialFileExtensions;
+		PotentialFileExtensions.Add(TEXT(".xls"));
+		PotentialFileExtensions.Add(TEXT(".xlsm"));
+		PotentialFileExtensions.Add(TEXT(".csv"));
+		PotentialFileExtensions.Add(TEXT(".json"));
+
+		UE::AssetTools::ExecuteFindSourceFileInExplorer(ImportPaths, PotentialFileExtensions);
+	}
+
+	bool CanExecuteFindSourceFileInExplorer(const FToolMenuContext& InContext)
+	{
+		const UContentBrowserAssetContextMenuContext* Context = UContentBrowserAssetContextMenuContext::FindContextWithAssets(InContext);
+
+		TArray<FString> ImportPaths;
+		for (const FAssetData& Asset : Context->SelectedAssets)
+		{
+			UAssetDefinitionRegistry::Get()->GetAssetDefinitionForAsset(Asset)->GetSourceFiles(Asset, [&ImportPaths](const FAssetImportInfo& AssetImportInfo)
+			{
+				for (const auto& SourceFile : AssetImportInfo.SourceFiles)
+				{
+					ImportPaths.Add(SourceFile.RelativeFilename);
+				}
+			});
+		}
+
+		TArray<FString> PotentialFileExtensions;
+		PotentialFileExtensions.Add(TEXT(".xls"));
+		PotentialFileExtensions.Add(TEXT(".xlsm"));
+		PotentialFileExtensions.Add(TEXT(".csv"));
+		PotentialFileExtensions.Add(TEXT(".json"));
+
+		return UE::AssetTools::CanExecuteFindSourceFileInExplorer(ImportPaths, PotentialFileExtensions);
+	}
+	
+	static FDelayedAutoRegisterHelper DelayedAutoRegister(EDelayedRegisterRunPhase::EndOfEngineInit, []{ 
+		UToolMenus::RegisterStartupCallback(FSimpleMulticastDelegate::FDelegate::CreateLambda([]()
+		{
+			FToolMenuOwnerScoped OwnerScoped(UE_MODULE_NAME);
+			UToolMenu* Menu = UE::ContentBrowser::ExtendToolMenu_AssetContextMenu(UDataTable::StaticClass());
+		
+			FToolMenuSection& Section = Menu->FindOrAddSection("GetAssetActions");
+			Section.AddDynamicEntry(NAME_None, FNewToolMenuSectionDelegate::CreateLambda([](FToolMenuSection& InSection)
+			{
+				{
+					const TAttribute<FText> Label = LOCTEXT("DataTable_ExportAsCSV", "Export as CSV");
+					const TAttribute<FText> ToolTip = LOCTEXT("DataTable_ExportAsCSVTooltip", "Export the data table as a file containing CSV data.");
+					const FSlateIcon Icon = FSlateIcon();
+
+					FToolUIAction UIAction;
+					UIAction.ExecuteAction = FToolMenuExecuteAction::CreateStatic(&ExecuteExportAsCSV);
+					InSection.AddMenuEntry("DataTable_ExportAsCSV", Label, ToolTip, Icon, UIAction);
+				}
+				
+				{
+					const TAttribute<FText> Label = LOCTEXT("DataTable_ExportAsJSON", "Export as JSON");
+					const TAttribute<FText> ToolTip = LOCTEXT("DataTable_ExportAsJSONTooltip", "Export the data table as a file containing JSON data.");
+					const FSlateIcon Icon = FSlateIcon();
+
+					FToolUIAction UIAction;
+					UIAction.ExecuteAction = FToolMenuExecuteAction::CreateStatic(&ExecuteExportAsJSON);
+					InSection.AddMenuEntry("DataTable_ExportAsJSON", Label, ToolTip, Icon, UIAction);
+				}
+				{
+					const TAttribute<FText> Label = LOCTEXT("DataTable_OpenSourceData", "Open Source Data");
+					const TAttribute<FText> ToolTip = LOCTEXT("DataTable_OpenSourceDataTooltip", "Opens the data table's source data file in an external editor. It will search using the following extensions: .xls/.xlsm/.csv/.json");
+					const FSlateIcon Icon = FSlateIcon();
+
+					FToolUIAction UIAction;
+					UIAction.ExecuteAction = FToolMenuExecuteAction::CreateLambda(&ExecuteFindSourceFileInExplorer);
+					UIAction.CanExecuteAction = FToolMenuCanExecuteAction::CreateStatic(&CanExecuteFindSourceFileInExplorer);
+					InSection.AddMenuEntry("DataTable_OpenSourceData", Label, ToolTip, Icon, UIAction);
+				}
+			}));
+		}));
+	});
+}
+
+/*
+void FAssetTypeActions_DataTable::GetActions(const TArray<UObject*>& InObjects, FToolMenuSection& Section)
+{
+	auto Tables = GetTypedWeakObjectPtrs<UObject>(InObjects);
+	
+	TArray<FString> ImportPaths;
+	for (auto TableIter = Tables.CreateConstIterator(); TableIter; ++TableIter)
+	{
+		const UDataTable* CurTable = Cast<UDataTable>((*TableIter).Get());
+		if (CurTable)
+		{
+			CurTable->AssetImportData->ExtractFilenames(ImportPaths);
+		}
+	}
+
+	UAssetDefinitionRegistry::Get()->GetAssetDefinitionForAsset(YourAsset)->GetSourceFiles()
+
+	TArray<FString> PotentialFileExtensions;
+	PotentialFileExtensions.Add(TEXT(".xls"));
+	PotentialFileExtensions.Add(TEXT(".xlsm"));
+	PotentialFileExtensions.Add(TEXT(".csv"));
+	PotentialFileExtensions.Add(TEXT(".json"));
+	Section.AddMenuEntry(
+		"DataTable_OpenSourceData",
+		LOCTEXT("DataTable_OpenSourceData", "Open Source Data"),
+		LOCTEXT("DataTable_OpenSourceDataTooltip", "Opens the data table's source data file in an external editor. It will search using the following extensions: .xls/.xlsm/.csv/.json"),
+		FSlateIcon(),
+		FUIAction(
+			FExecuteAction::CreateSP( this, &ExecuteFindSourceFileInExplorer, ImportPaths, PotentialFileExtensions ),
+			FCanExecuteAction::CreateSP(this, &CanExecuteFindSourceFileInExplorer, ImportPaths, PotentialFileExtensions)
+			)
+		); 
+}
+*/
+
+#undef LOCTEXT_NAMESPACE
diff --git a/Engine/Plugins/Editor/EngineAssetDefinitions/Source/Private/Table/AssetDefinition_DataTable.h b/Engine/Plugins/Editor/EngineAssetDefinitions/Source/Private/Table/AssetDefinition_DataTable.h
new file mode 100644
--- /dev/null
+++ b/Engine/Plugins/Editor/EngineAssetDefinitions/Source/Private/Table/AssetDefinition_DataTable.h
@@ -0,0 +1,30 @@
+// Copyright Epic Games, Inc. All Rights Reserved.
+
+#pragma once
+
+#include "Engine/DataTable.h"
+#include "AssetDefinitionDefault.h"
+
+#include "AssetDefinition_DataTable.generated.h"
+
+UCLASS()
+class UAssetDefinition_DataTable : public UAssetDefinitionDefault
+{
+	GENERATED_BODY()
+
+public:
+	// UAssetDefinition Begin
+	virtual FText GetAssetDisplayName() const override { return NSLOCTEXT("AssetTypeActions", "AssetTypeActions_DataTable", "Data Table"); }
+	virtual FText GetAssetDisplayName(const FAssetData& AssetData) const override;
+	virtual FLinearColor GetAssetColor() const override { return FLinearColor(FColor(62, 140, 35)); }
+	virtual TSoftClassPtr<UObject> GetAssetClass() const override { return UDataTable::StaticClass(); }
+	virtual TConstArrayView<FAssetCategoryPath> GetAssetCategories() const override
+	{
+		static const auto Categories = { EAssetCategoryPaths::Misc };
+		return Categories;
+	}
+	virtual bool CanImport() const override { return true; }
+	virtual EAssetCommandResult PerformAssetDiff(const FAssetDiffArgs& DiffArgs) const override;
+	virtual EAssetCommandResult OpenAssets(const FAssetOpenArgs& OpenArgs) const override;
+	// UAssetDefinition End
+};
diff --git a/Engine/Plugins/Editor/EngineAssetDefinitions/Source/Private/Table/AssetDefinition_MirrorDataTable.h b/Engine/Plugins/Editor/EngineAssetDefinitions/Source/Private/Table/AssetDefinition_MirrorDataTable.h
new file mode 100644
--- /dev/null
+++ b/Engine/Plugins/Editor/EngineAssetDefinitions/Source/Private/Table/AssetDefinition_MirrorDataTable.h
@@ -0,0 +1,26 @@
+// Copyright Epic Games, Inc. All Rights Reserved.
+
+#pragma once
+
+#include "Animation/MirrorDataTable.h"
+#include "AssetDefinition_DataTable.h"
+
+#include "AssetDefinition_MirrorDataTable.generated.h"
+
+UCLASS()
+class UAssetDefinition_MirrorDataTable : public UAssetDefinition_DataTable
+{
+	GENERATED_BODY()
+
+public:
+	// UAssetDefinition Begin
+	virtual FText GetAssetDisplayName() const override { return NSLOCTEXT("UMirrorDataTable", "FAssetTypeActions_MirrorDataTable", "Mirror Data Table"); }
+	virtual TSoftClassPtr<UObject> GetAssetClass() const override { return UMirrorDataTable::StaticClass(); }
+	virtual TConstArrayView<FAssetCategoryPath> GetAssetCategories() const override
+	{
+		static const auto Categories = { EAssetCategoryPaths::Animation };
+		return Categories;
+	}
+	virtual bool CanImport() const override { return false; }
+	// UAssetDefinition End
+};
diff --git a/Engine/Source/Developer/AssetTools/Private/AssetTools.cpp b/Engine/Source/Developer/AssetTools/Private/AssetTools.cpp
--- a/Engine/Source/Developer/AssetTools/Private/AssetTools.cpp
+++ b/Engine/Source/Developer/AssetTools/Private/AssetTools.cpp
@@ -51,15 +51,10 @@
 #include "AssetTypeActions/AssetTypeActions_AimOffset.h"
 #include "AssetTypeActions/AssetTypeActions_BlendSpace1D.h"
 #include "AssetTypeActions/AssetTypeActions_AimOffset1D.h"
-#include "AssetTypeActions/AssetTypeActions_CurveTable.h"
-#include "AssetTypeActions/AssetTypeActions_CompositeCurveTable.h"
 #include "AssetTypeActions/AssetTypeActions_DataAsset.h"
-#include "AssetTypeActions/AssetTypeActions_DataTable.h"
-#include "AssetTypeActions/AssetTypeActions_CompositeDataTable.h"
 #include "AssetTypeActions/AssetTypeActions_Class.h"
 #include "AssetTypeActions/AssetTypeActions_ForceFeedbackEffect.h"
 #include "AssetTypeActions/AssetTypeActions_ActorFoliageSettings.h"
-#include "AssetTypeActions/AssetTypeActions_MirrorDataTable.h"
 #include "AssetTypeActions/AssetTypeActions_ParticleSystem.h"
 #include "AssetTypeActions/AssetTypeActions_PhysicalMaterialMask.h"
 #include "WorldPartition/WorldPartition.h"
@@ -1345,14 +1340,9 @@ UAssetToolsImpl::UAssetToolsImpl(const FObjectInitializer& ObjectInitializer)
 	RegisterAssetTypeActions(MakeShareable(new FAssetTypeActions_BlendSpace1D));
 	RegisterAssetTypeActions(MakeShareable(new FAssetTypeActions_Blueprint));
 	RegisterAssetTypeActions(MakeShareable(new FAssetTypeActions_BlueprintGeneratedClass));
-	RegisterAssetTypeActions(MakeShareable(new FAssetTypeActions_CurveTable));
-	RegisterAssetTypeActions(MakeShareable(new FAssetTypeActions_CompositeCurveTable));
-	RegisterAssetTypeActions(MakeShareable(new FAssetTypeActions_DataTable));
-	RegisterAssetTypeActions(MakeShareable(new FAssetTypeActions_CompositeDataTable));
 	RegisterAssetTypeActions(MakeShareable(new FAssetTypeActions_Class));
 	RegisterAssetTypeActions(MakeShareable(new FAssetTypeActions_ForceFeedbackEffect(InputCategoryBit)));
 	RegisterAssetTypeActions(MakeShareable(new FAssetTypeActions_ActorFoliageSettings(FoliageCategoryBit)));
-	RegisterAssetTypeActions(MakeShareable(new FAssetTypeActions_MirrorDataTable));
 	RegisterAssetTypeActions(MakeShareable(new FAssetTypeActions_ParticleSystem));
 	RegisterAssetTypeActions(MakeShareable(new FAssetTypeActions_PhysicalMaterialMask));
 
diff --git a/Engine/Source/Developer/AssetTools/Private/AssetTypeActions/AssetTypeActions_CompositeCurveTable.h b/Engine/Source/Developer/AssetTools/Private/AssetTypeActions/AssetTypeActions_CompositeCurveTable.h
deleted file mode 100644
diff --git a/Engine/Source/Developer/AssetTools/Private/AssetTypeActions/AssetTypeActions_CompositeDataTable.h b/Engine/Source/Developer/AssetTools/Private/AssetTypeActions/AssetTypeActions_CompositeDataTable.h
deleted file mode 100644
diff --git a/Engine/Source/Developer/AssetTools/Private/AssetTypeActions/AssetTypeActions_CurveTable.cpp b/Engine/Source/Developer/AssetTools/Private/AssetTypeActions/AssetTypeActions_CurveTable.cpp
deleted file mode 100644
diff --git a/Engine/Source/Developer/AssetTools/Private/AssetTypeActions/AssetTypeActions_CurveTable.h b/Engine/Source/Developer/AssetTools/Private/AssetTypeActions/AssetTypeActions_CurveTable.h
deleted file mode 100644
diff --git a/Engine/Source/Developer/AssetTools/Private/AssetTypeActions/AssetTypeActions_DataTable.cpp b/Engine/Source/Developer/AssetTools/Private/AssetTypeActions/AssetTypeActions_DataTable.cpp
deleted file mode 100644
diff --git a/Engine/Source/Developer/AssetTools/Private/AssetTypeActions/AssetTypeActions_DataTable.h b/Engine/Source/Developer/AssetTools/Private/AssetTypeActions/AssetTypeActions_DataTable.h
deleted file mode 100644
diff --git a/Engine/Source/Developer/AssetTools/Private/AssetTypeActions/AssetTypeActions_MirrorDataTable.h b/Engine/Source/Developer/AssetTools/Private/AssetTypeActions/AssetTypeActions_MirrorDataTable.h
deleted file mode 100644
diff --git a/Engine/Source/Developer/AssetTools/Private/FindSourceFileInExplorer.cpp b/Engine/Source/Developer/AssetTools/Private/FindSourceFileInExplorer.cpp
new file mode 100644
--- /dev/null
+++ b/Engine/Source/Developer/AssetTools/Private/FindSourceFileInExplorer.cpp
@@ -0,0 +1,58 @@
+// Copyright Epic Games, Inc. All Rights Reserved.
+
+#include "FindSourceFileInExplorer.h"
+#include "Misc/Paths.h"
+#include "HAL/PlatformProcess.h"
+
+namespace UE::AssetTools
+{
+	void ExecuteFindSourceFileInExplorer(TArray<FString> Filenames, TArray<FString> OverrideExtensions)
+	{
+		for (TArray<FString>::TConstIterator FilenameIter(Filenames); FilenameIter; ++FilenameIter)
+		{
+			const FString CSVFilename = FPaths::ConvertRelativePathToFull(*FilenameIter);
+			const FString RootPath = FPaths::GetPath(CSVFilename);
+			const FString BaseFilename = FPaths::GetBaseFilename(CSVFilename, true);
+		
+			for (TArray<FString>::TConstIterator ExtensionItr(OverrideExtensions); ExtensionItr; ++ExtensionItr)
+			{
+				const FString FilenameWithExtension(FString::Printf(TEXT("%s/%s%s"), *RootPath, *BaseFilename, **ExtensionItr));
+			
+				if (!FilenameWithExtension.IsEmpty() && FPaths::FileExists(*FilenameWithExtension))
+				{
+					FPlatformProcess::LaunchFileInDefaultExternalApplication(*FilenameWithExtension, nullptr, ELaunchVerb::Edit);
+					break;
+				}
+			}
+		}
+	}
+
+	bool CanExecuteFindSourceFileInExplorer(TArray<FString> Filenames, TArray<FString> OverrideExtensions)
+	{
+		// Verify that extensions were provided
+		if (OverrideExtensions.Num() == 0)
+		{
+			return false;
+		}
+
+		// Verify that the file exists with any of the given extensions
+		for (TArray<FString>::TConstIterator FilenameIter(Filenames); FilenameIter; ++FilenameIter)
+		{
+			const FString CSVFilename = FPaths::ConvertRelativePathToFull(*FilenameIter);
+			const FString RootPath = FPaths::GetPath(CSVFilename);
+			const FString BaseFilename = FPaths::GetBaseFilename(CSVFilename, true);
+
+			for (TArray<FString>::TConstIterator ExtensionItr(OverrideExtensions); ExtensionItr; ++ExtensionItr)
+			{
+				const FString FilenameWithExtension(FString::Printf(TEXT("%s/%s%s"), *RootPath, *BaseFilename, **ExtensionItr));
+
+				if (!FilenameWithExtension.IsEmpty() && FPaths::FileExists(*FilenameWithExtension))
+				{
+					return true;
+				}
+			}
+		}
+
+		return false;
+	}
+}
\ No newline at end of file
diff --git a/Engine/Source/Developer/AssetTools/Public/FindSourceFileInExplorer.h b/Engine/Source/Developer/AssetTools/Public/FindSourceFileInExplorer.h
new file mode 100644
--- /dev/null
+++ b/Engine/Source/Developer/AssetTools/Public/FindSourceFileInExplorer.h
@@ -0,0 +1,15 @@
+// Copyright Epic Games, Inc. All Rights Reserved.
+
+#pragma once
+
+#include "CoreMinimal.h"
+
+//virtual FColor GetTypeColor() const override { return FColor(62, 140, 35); }
+//virtual uint32 GetCategories() override { return EAssetTypeCategories::Misc; }
+//virtual bool IsImportedAsset() const override { return true; }
+
+namespace UE::AssetTools
+{
+	ASSETTOOLS_API void ExecuteFindSourceFileInExplorer(TArray<FString> Filenames, TArray<FString> OverrideExtensions);
+	ASSETTOOLS_API bool CanExecuteFindSourceFileInExplorer(TArray<FString> Filenames, TArray<FString> OverrideExtensions);
+}
diff --git a/Engine/Source/Editor/AssetDefinition/Private/AssetDefinition.cpp b/Engine/Source/Editor/AssetDefinition/Private/AssetDefinition.cpp
--- a/Engine/Source/Editor/AssetDefinition/Private/AssetDefinition.cpp
+++ b/Engine/Source/Editor/AssetDefinition/Private/AssetDefinition.cpp
@@ -3,6 +3,7 @@
 #include "AssetDefinition.h"
 #include "AssetDefinitionRegistry.h"
 #include "Misc/AssetFilterData.h"
+#include "EditorFramework/AssetImportData.h"
 
 #define LOCTEXT_NAMESPACE "UAssetDefinition"
 
@@ -67,6 +68,23 @@ TConstArrayView<FAssetCategoryPath> UAssetDefinition::GetAssetCategories() const
 	return Categories;
 }
 
+EAssetCommandResult UAssetDefinition::GetSourceFiles(const FAssetData& InAsset, TFunctionRef<void(const FAssetImportInfo& AssetImportData)> SourceFileFunc) const
+{
+	FString SourceFileTagData;
+	if (InAsset.GetTagValue(UObject::SourceFileTagName(), SourceFileTagData))
+	{
+		TOptional<FAssetImportInfo> ImportInfo = FAssetImportInfo::FromJson(SourceFileTagData);
+		if (ImportInfo.IsSet())
+		{
+			SourceFileFunc(ImportInfo.GetValue());
+			
+			return EAssetCommandResult::Handled;
+		}
+	}
+    
+	return EAssetCommandResult::Unhandled;
+}
+
 bool UAssetDefinition::CanRegisterStatically() const
 {
 	return !GetClass()->HasAnyClassFlags(CLASS_Abstract);
diff --git a/Engine/Source/Editor/AssetDefinition/Public/AssetDefinition.h b/Engine/Source/Editor/AssetDefinition/Public/AssetDefinition.h
--- a/Engine/Source/Editor/AssetDefinition/Public/AssetDefinition.h
+++ b/Engine/Source/Editor/AssetDefinition/Public/AssetDefinition.h
@@ -109,11 +109,6 @@ struct FAssetActivateArgs : public FAssetArgs
 	EAssetActivationMethod ActivationMethod;
 };
 
-struct FAssetSourceFileArgs : public FAssetArgs
-{
-	FAssetSourceFileArgs(TConstArrayView<FAssetData> InAssets) : FAssetArgs(InAssets) { }
-};
-
 struct FAssetSourceFile
 {
 	FString DisplayLabelName;
@@ -278,8 +273,8 @@ public:
 	TOptional<EToolkitMode::Type> RequiredToolkitMode;
 };
 
-
 class UAssetDefinitionRegistry;
+struct FAssetImportInfo;
 
 enum class EIncludeClassInFilter : uint8
 {
@@ -385,10 +380,7 @@ public:
 	}
 
 	// Source Files
-	virtual EAssetCommandResult GetSourceFiles(const FAssetSourceFileArgs& SourceFileArgs, TArray<FAssetSourceFile>& OutSourceAssets) const
-	{
-		return EAssetCommandResult::Unhandled;
-	}
+	virtual EAssetCommandResult GetSourceFiles(const FAssetData& InAsset, TFunctionRef<void(const FAssetImportInfo& AssetImportData)> SourceFileFunc) const;
 
 	// Diffing Assets
 	virtual EAssetCommandResult PerformAssetDiff(const FAssetDiffArgs& DiffArgs) const
diff --git a/Engine/Source/Editor/UnrealEd/Private/AssetDefinitionDefault.cpp b/Engine/Source/Editor/UnrealEd/Private/AssetDefinitionDefault.cpp
--- a/Engine/Source/Editor/UnrealEd/Private/AssetDefinitionDefault.cpp
+++ b/Engine/Source/Editor/UnrealEd/Private/AssetDefinitionDefault.cpp
@@ -20,30 +20,6 @@ EAssetCommandResult UAssetDefinitionDefault::OpenAssets(const FAssetOpenArgs& Op
 	return EAssetCommandResult::Unhandled;
 }
 
-EAssetCommandResult UAssetDefinitionDefault::GetSourceFiles(const FAssetSourceFileArgs& SourceFileArgs, TArray<FAssetSourceFile>& OutSourceAssets) const
-{
-	for (const FAssetData& Asset : SourceFileArgs.Assets)
-    {
-    	FString SourceFileTagData;
-    	if (Asset.GetTagValue(UObject::SourceFileTagName(), SourceFileTagData))
-    	{
-    		TOptional<FAssetImportInfo> ImportInfo = FAssetImportInfo::FromJson(SourceFileTagData);
-    		if (ImportInfo.IsSet())
-    		{
-    			for (const FAssetImportInfo::FSourceFile& SourceFile : ImportInfo->SourceFiles)
-    			{
-    				FAssetSourceFile AssetSourceFile;
-    				AssetSourceFile.DisplayLabelName = SourceFile.DisplayLabelName;
-    				AssetSourceFile.RelativeFilename = SourceFile.RelativeFilename;
-    				OutSourceAssets.Add(MoveTemp(AssetSourceFile));
-    			}
-    		}
-    	}
-    }
-    
-    return EAssetCommandResult::Handled;
-}
-
 EAssetCommandResult UAssetDefinitionDefault::PerformAssetDiff(const FAssetDiffArgs& DiffArgs) const
 {
 	check(DiffArgs.OldAsset != nullptr);
@@ -61,20 +37,4 @@ EAssetCommandResult UAssetDefinitionDefault::PerformAssetDiff(const FAssetDiffAr
 	return EAssetCommandResult::Handled;
 }
 
-/*
-virtual void PerformAssetDiff(UObject* OldAsset, UObject* NewAsset, const struct FRevisionInfo& OldRevision, const struct FRevisionInfo& NewRevision) const override
-{
-	check(OldAsset != nullptr);
-	check(NewAsset != nullptr);
-
-	// Dump assets to temp text files
-	FString OldTextFilename = DumpAssetToTempFile(OldAsset);
-	FString NewTextFilename = DumpAssetToTempFile(NewAsset);
-	FString DiffCommand = GetDefault<UEditorLoadingSavingSettings>()->TextDiffToolPath.FilePath;
-
-	FAssetToolsModule& AssetToolsModule = FModuleManager::Get().LoadModuleChecked<FAssetToolsModule>("AssetTools");
-	AssetToolsModule.Get().CreateDiffProcess(DiffCommand, OldTextFilename, NewTextFilename);
-}
-*/
-
 #undef LOCTEXT_NAMESPACE
diff --git a/Engine/Source/Editor/UnrealEd/Private/Editor.cpp b/Engine/Source/Editor/UnrealEd/Private/Editor.cpp
--- a/Engine/Source/Editor/UnrealEd/Private/Editor.cpp
+++ b/Engine/Source/Editor/UnrealEd/Private/Editor.cpp
@@ -59,6 +59,8 @@
 #include "Interfaces/IAnalyticsProvider.h"
 #include "EngineAnalytics.h"
 
+#include "EditorFramework/AssetImportData.h"
+
 // AIMdule
 
 #include "AssetDefinition.h"
@@ -838,9 +840,11 @@ void FReimportManager::GetNewReimportPath(UObject* Obj, TArray<FString>& InOutFi
 			FAssetData AssetData(Obj);
 			const UAssetDefinition* AssetDefinition = UAssetDefinitionRegistry::Get()->GetAssetDefinitionForAsset(AssetData);
 
-			TArray<FAssetSourceFile> OutSourceAssets;
-			FAssetSourceFileArgs SourceFileArgs({ AssetData });
-			AssetDefinition->GetSourceFiles(SourceFileArgs, OutSourceAssets);
+			TArray<FAssetImportInfo::FSourceFile> OutSourceAssets;
+			AssetDefinition->GetSourceFiles(AssetData, [&OutSourceAssets](const FAssetImportInfo& ImportInfo)
+			{
+				OutSourceAssets.Append(ImportInfo.SourceFiles);
+			});
 			
 			if (OutSourceAssets.IsValidIndex(SourceFileIndex))
 			{
diff --git a/Engine/Source/Editor/UnrealEd/Public/AssetDefinitionDefault.h b/Engine/Source/Editor/UnrealEd/Public/AssetDefinitionDefault.h
--- a/Engine/Source/Editor/UnrealEd/Public/AssetDefinitionDefault.h
+++ b/Engine/Source/Editor/UnrealEd/Public/AssetDefinitionDefault.h
@@ -15,7 +15,6 @@ class UNREALED_API UAssetDefinitionDefault : public UAssetDefinition
 public:
 	// UAssetDefinition Begin
 	virtual EAssetCommandResult OpenAssets(const FAssetOpenArgs& OpenArgs) const override;
-	virtual EAssetCommandResult GetSourceFiles(const FAssetSourceFileArgs& SourceFileArgs, TArray<FAssetSourceFile>& OutSourceAssets) const override;
 	virtual EAssetCommandResult PerformAssetDiff(const FAssetDiffArgs& DiffArgs) const override;
 	// UAssetDefinition End
 };
