AssetDefinition - Upgrading several skeletal and animation assets to be asset definitions.  Had to upgrade a few utilities that were causing the assets to still load.

diff --git a/Engine/Plugins/Animation/ControlRig/Source/ControlRigEditor/Private/ControlRigEditorModule.cpp b/Engine/Plugins/Animation/ControlRig/Source/ControlRigEditor/Private/ControlRigEditorModule.cpp
--- a/Engine/Plugins/Animation/ControlRig/Source/ControlRigEditor/Private/ControlRigEditorModule.cpp
+++ b/Engine/Plugins/Animation/ControlRig/Source/ControlRigEditor/Private/ControlRigEditorModule.cpp
@@ -913,76 +913,59 @@ void FControlRigEditorModule::AddControlRigExtenderToToolMenu(FName InToolMenuNa
 	);
 }
 
-void FControlRigEditorModule::ExtendAnimSequenceMenu()
+static void ExecuteOpenLevelSequence(const FToolMenuContext& InContext)
 {
-	TArray<UToolMenu*> MenusToExtend;
-	MenusToExtend.Add(UToolMenus::Get()->ExtendMenu("ContentBrowser.AssetContextMenu.AnimSequence"));
+	if (const UContentBrowserAssetContextMenuContext* CBContext = InContext.FindContext<UContentBrowserAssetContextMenuContext>())
+	{
+		if (UAnimSequence* AnimSequence = CBContext->LoadFirstSelectedObject<UAnimSequence>())
+        {
+       		FControlRigEditorModule::OpenLevelSequence(AnimSequence);
+        }
+	}
+}
 
-	for (UToolMenu* Menu : MenusToExtend)
+static bool CanExecuteOpenLevelSequence(const FToolMenuContext& InContext)
+{
+	if (const UContentBrowserAssetContextMenuContext* CBContext = InContext.FindContext<UContentBrowserAssetContextMenuContext>())
 	{
-		if (Menu == nullptr)
+		if (CBContext->SelectedAssets.Num() != 1)
 		{
-			continue;
+			return false;
 		}
 
-		FToolMenuSection& Section = Menu->FindOrAddSection("GetAssetActions");
-		Section.AddDynamicEntry("ControlRigOpenLevelSequence", FNewToolMenuSectionDelegate::CreateLambda([](FToolMenuSection& InSection)
-			{
-				UContentBrowserAssetContextMenuContext* Context = InSection.FindContext<UContentBrowserAssetContextMenuContext>();
-				if (Context)
-				{
-					TArray<UObject*> SelectedObjects = Context->GetSelectedObjects();
-					if (SelectedObjects.Num() > 0)
-					{
-						InSection.AddMenuEntry(
-							"OpenLevelSequence",
-							LOCTEXT("OpenLevelSequence", "Open Level Sequence"),
-							LOCTEXT("CreateControlRig_ToolTip", "Opens a Level Sequence if it is driving this Anim Sequence."),
-							FSlateIcon(FAppStyle::GetAppStyleSetName(), "GenericCurveEditor.TabIcon"),
-							FUIAction(
-								FExecuteAction::CreateLambda([SelectedObjects]()
-									{
-										for (UObject* SelectedObject : SelectedObjects)
-										{
-											UAnimSequence* AnimSequence = Cast<UAnimSequence>(SelectedObject);
-											if (AnimSequence)
-											{
-												FControlRigEditorModule::OpenLevelSequence(AnimSequence);
-												return; //just open up the first valid one, can't have more than one open.
-											}
-										}
-									}),
-								FCanExecuteAction::CreateLambda([SelectedObjects]()
-									{
-										for (UObject* SelectedObject : SelectedObjects)
-										{
-											UAnimSequence* AnimSequence = Cast<UAnimSequence>(SelectedObject);
-											if (AnimSequence)
-											{
-												if (IInterface_AssetUserData* AnimAssetUserData = Cast< IInterface_AssetUserData >(AnimSequence))
-												{
-													UAnimSequenceLevelSequenceLink* AnimLevelLink = AnimAssetUserData->GetAssetUserData< UAnimSequenceLevelSequenceLink >();
-													if (AnimLevelLink)
-													{
-														ULevelSequence* LevelSequence = AnimLevelLink->ResolveLevelSequence();
-														if (LevelSequence)
-														{
-															return true;
-														}
-													}
-												}
-											}
-										}
-										return false;
-									})
+		const FAssetData& SelectedAnimSequence = CBContext->SelectedAssets[0];
 
-								)
-						);
-							
-					}
-				}
-			}));
+		FString PathToLevelSequence;
+		if (SelectedAnimSequence.GetTagValue<FString>(GET_MEMBER_NAME_CHECKED(UAnimSequenceLevelSequenceLink, PathToLevelSequence), PathToLevelSequence))
+		{
+			if (!FSoftObjectPath(PathToLevelSequence).IsNull())
+			{
+				return true;
+			}
+		}
 	}
+
+	return false;
+}
+
+void FControlRigEditorModule::ExtendAnimSequenceMenu()
+{
+	UToolMenu* Menu = UE::ContentBrowser::ExtendToolMenu_AssetContextMenu(UAnimSequence::StaticClass());
+
+	FToolMenuSection& Section = Menu->FindOrAddSection("GetAssetActions");
+	Section.AddDynamicEntry(NAME_None, FNewToolMenuSectionDelegate::CreateLambda([](FToolMenuSection& InSection)
+		{
+			{
+				const TAttribute<FText> Label = LOCTEXT("OpenLevelSequence", "Open Level Sequence");
+				const TAttribute<FText> ToolTip =LOCTEXT("CreateControlRig_ToolTip", "Opens a Level Sequence if it is driving this Anim Sequence.");
+				const FSlateIcon Icon = FSlateIcon(FAppStyle::GetAppStyleSetName(), "GenericCurveEditor.TabIcon");
+
+				FToolUIAction UIAction;
+				UIAction.ExecuteAction = FToolMenuExecuteAction::CreateStatic(&ExecuteOpenLevelSequence);
+				UIAction.CanExecuteAction = FToolMenuCanExecuteAction::CreateStatic(&CanExecuteOpenLevelSequence);
+				InSection.AddMenuEntry("OpenLevelSequence", Label, ToolTip, Icon, UIAction);
+			}
+		}));
 }
 
 void FControlRigEditorModule::HandleNewBlueprintCreated(UBlueprint* InBlueprint)
diff --git a/Engine/Plugins/Editor/EngineAssetDefinitions/Source/Private/Animation/AssetDefinition_AimOffset.h b/Engine/Plugins/Editor/EngineAssetDefinitions/Source/Private/Animation/AssetDefinition_AimOffset.h
new file mode 100644
--- /dev/null
+++ b/Engine/Plugins/Editor/EngineAssetDefinitions/Source/Private/Animation/AssetDefinition_AimOffset.h
@@ -0,0 +1,21 @@
+// Copyright Epic Games, Inc. All Rights Reserved.
+
+#pragma once
+
+#include "Animation/AimOffsetBlendSpace.h"
+#include "AssetDefinition_BlendSpace.h"
+
+#include "AssetDefinition_AimOffset.generated.h"
+
+UCLASS()
+class UAssetDefinition_AimOffset : public UAssetDefinition_BlendSpace
+{
+	GENERATED_BODY()
+
+public:
+	// UAssetDefinition Begin
+	virtual FText GetAssetDisplayName() const override { return NSLOCTEXT("AssetTypeActions", "AssetTypeActions_AimOffset", "Aim Offset"); }
+	virtual FLinearColor GetAssetColor() const override { return FLinearColor(FColor(0, 162, 232)); }
+	virtual TSoftClassPtr<UObject> GetAssetClass() const override { return UAimOffsetBlendSpace::StaticClass(); }
+	// UAssetDefinition End
+};
\ No newline at end of file
diff --git a/Engine/Plugins/Editor/EngineAssetDefinitions/Source/Private/Animation/AssetDefinition_AimOffset1D.h b/Engine/Plugins/Editor/EngineAssetDefinitions/Source/Private/Animation/AssetDefinition_AimOffset1D.h
new file mode 100644
--- /dev/null
+++ b/Engine/Plugins/Editor/EngineAssetDefinitions/Source/Private/Animation/AssetDefinition_AimOffset1D.h
@@ -0,0 +1,26 @@
+// Copyright Epic Games, Inc. All Rights Reserved.
+
+#pragma once
+
+#include "Animation/AimOffsetBlendSpace1D.h"
+#include "AssetDefinition_BlendSpace1D.h"
+
+#include "AssetDefinition_AimOffset1D.generated.h"
+
+UCLASS()
+class UAssetDefinition_AimOffset1D : public UAssetDefinition_BlendSpace1D
+{
+	GENERATED_BODY()
+
+public:
+	// UAssetDefinition Begin
+	virtual FText GetAssetDisplayName() const override { return NSLOCTEXT("AssetTypeActions", "AssetTypeActions_AimOffset1D", "Aim Offset 1D"); }
+	virtual FLinearColor GetAssetColor() const override { return FLinearColor(FColor(0, 162, 232)); }
+	virtual TSoftClassPtr<UObject> GetAssetClass() const override { return UAimOffsetBlendSpace1D::StaticClass(); }
+	virtual TConstArrayView<FAssetCategoryPath> GetAssetCategories() const override
+	{
+		static const auto Categories = { EAssetCategoryPaths::Animation / NSLOCTEXT("AssetTypeActions", "AnimLegacySubMenu", "Legacy") };
+		return Categories;
+	}
+	// UAssetDefinition End
+};
\ No newline at end of file
diff --git a/Engine/Plugins/Editor/EngineAssetDefinitions/Source/Private/Animation/AssetDefinition_AnimComposite.h b/Engine/Plugins/Editor/EngineAssetDefinitions/Source/Private/Animation/AssetDefinition_AnimComposite.h
new file mode 100644
--- /dev/null
+++ b/Engine/Plugins/Editor/EngineAssetDefinitions/Source/Private/Animation/AssetDefinition_AnimComposite.h
@@ -0,0 +1,21 @@
+// Copyright Epic Games, Inc. All Rights Reserved.
+
+#pragma once
+
+#include "Animation/AnimComposite.h"
+#include "AssetDefinition_AnimationAsset.h"
+
+#include "AssetDefinition_AnimComposite.generated.h"
+
+UCLASS()
+class UAssetDefinition_AnimComposite : public UAssetDefinition_AnimationAsset
+{
+	GENERATED_BODY()
+
+public:
+	// UAssetDefinition Begin
+	virtual FText GetAssetDisplayName() const override { return NSLOCTEXT("AssetTypeActions", "AssetTypeActions_AnimComposite", "Animation Composite"); }
+	virtual FLinearColor GetAssetColor() const override { return FLinearColor(FColor(181, 230, 29)); }
+	virtual TSoftClassPtr<UObject> GetAssetClass() const override { return UAnimComposite::StaticClass(); }
+	// UAssetDefinition End
+};
diff --git a/Engine/Plugins/Editor/EngineAssetDefinitions/Source/Private/Animation/AssetDefinition_AnimMontage.cpp b/Engine/Plugins/Editor/EngineAssetDefinitions/Source/Private/Animation/AssetDefinition_AnimMontage.cpp
new file mode 100644
--- /dev/null
+++ b/Engine/Plugins/Editor/EngineAssetDefinitions/Source/Private/Animation/AssetDefinition_AnimMontage.cpp
@@ -0,0 +1,90 @@
+// Copyright Epic Games, Inc. All Rights Reserved.
+
+#include "AssetDefinition_AnimMontage.h"
+
+#include "AnimationEditorUtils.h"
+#include "ContentBrowserMenuContexts.h"
+#include "ToolMenu.h"
+#include "ToolMenus.h"
+
+#define LOCTEXT_NAMESPACE "AssetTypeActions"
+
+// Menu Extensions
+//--------------------------------------------------------------------
+
+namespace MenuExtension_AnimMontage
+{
+	const static FName NAME_HasParentAsset("HasParentAsset");
+	static TOptional<bool> HasParentAsset(const FAssetData& AnimMontage)
+	{
+		bool bHasParentAsset = false;
+		if (!AnimMontage.GetTagValue<bool>(NAME_HasParentAsset, bHasParentAsset))
+		{
+			return TOptional<bool>();
+		}
+		
+		return bHasParentAsset;
+	}
+
+	void ExecuteCreateChildAnimMontage(const FToolMenuContext& InContext)
+	{
+		const UContentBrowserAssetContextMenuContext* Context = UContentBrowserAssetContextMenuContext::FindContextWithAssets(InContext);
+		TArray<UAnimMontage*> AnimMontages = Context->LoadSelectedObjects<UAnimMontage>();
+		
+        // only show child montage if inobjects are not child montage already
+        bool bContainsChildMontage = false;
+        for (UAnimMontage* AnimMontage : AnimMontages)
+        {
+        	bContainsChildMontage |= AnimMontage->HasParentAsset();
+        	if (bContainsChildMontage)
+        	{
+        		break;
+        	}
+        }
+		
+		if (!bContainsChildMontage && AnimMontages.Num() > 0)
+		{
+			TArray<UObject*> ObjectsToSync;
+			// need to know source and target
+			 for (UAnimMontage* ParentMontage : AnimMontages)
+             {
+				UAnimMontage* NewAsset = AnimationEditorUtils::CreateAnimationAsset<UAnimMontage>(ParentMontage->GetSkeleton(), ParentMontage->GetOutermost()->GetName(), TEXT("_Child"));
+				if (NewAsset)
+				{
+					NewAsset->SetParentAsset(ParentMontage);
+					ObjectsToSync.Add(NewAsset);
+				}
+			}
+
+			IAssetTools::Get().SyncBrowserToAssets(ObjectsToSync);
+		}
+	}
+	
+	static FDelayedAutoRegisterHelper DelayedAutoRegister(EDelayedRegisterRunPhase::EndOfEngineInit, []{ 
+		UToolMenus::RegisterStartupCallback(FSimpleMulticastDelegate::FDelegate::CreateLambda([]()
+		{
+			FToolMenuOwnerScoped OwnerScoped(UE_MODULE_NAME);
+			UToolMenu* Menu = UE::ContentBrowser::ExtendToolMenu_AssetContextMenu(UAnimMontage::StaticClass());
+		
+			FToolMenuSection& Section = Menu->FindOrAddSection("GetAssetActions");
+			Section.AddDynamicEntry(NAME_None, FNewToolMenuSectionDelegate::CreateLambda([](FToolMenuSection& InSection)
+			{
+				const UContentBrowserAssetContextMenuContext* CBContext = UContentBrowserAssetContextMenuContext::FindContextWithAssets(InSection);
+
+				// Only show create child montage option if none of the selected assets are child montages. 
+				if (!CBContext->SelectedAssets.ContainsByPredicate([](const FAssetData& InAsset) { return HasParentAsset(InAsset).Get(false); }))
+				{
+					const TAttribute<FText> Label = LOCTEXT("AnimMontage_CreateChildMontage", "Create Child Montage");
+					const TAttribute<FText> ToolTip = LOCTEXT("AnimMontage_CreateChildMontageTooltip", "Create Child Animation Montage and remap to another animation assets.");
+					const FSlateIcon Icon = FSlateIcon(FAppStyle::GetAppStyleSetName(), "ClassIcon.AnimMontage");
+
+					FToolUIAction UIAction;
+					UIAction.ExecuteAction = FToolMenuExecuteAction::CreateStatic(&ExecuteCreateChildAnimMontage);
+					InSection.AddMenuEntry("AnimMontage_CreateChildMontage", Label, ToolTip, Icon, UIAction);
+				}
+			}));
+		}));
+	});
+}
+
+#undef LOCTEXT_NAMESPACE
diff --git a/Engine/Plugins/Editor/EngineAssetDefinitions/Source/Private/Animation/AssetDefinition_AnimMontage.h b/Engine/Plugins/Editor/EngineAssetDefinitions/Source/Private/Animation/AssetDefinition_AnimMontage.h
new file mode 100644
--- /dev/null
+++ b/Engine/Plugins/Editor/EngineAssetDefinitions/Source/Private/Animation/AssetDefinition_AnimMontage.h
@@ -0,0 +1,21 @@
+// Copyright Epic Games, Inc. All Rights Reserved.
+
+#pragma once
+
+#include "Animation/AnimMontage.h"
+#include "AssetDefinition_AnimationAsset.h"
+
+#include "AssetDefinition_AnimMontage.generated.h"
+
+UCLASS()
+class UAssetDefinition_AnimMontage : public UAssetDefinition_AnimationAsset
+{
+	GENERATED_BODY()
+
+public:
+	// UAssetDefinition Begin
+	virtual FText GetAssetDisplayName() const override { return NSLOCTEXT("AssetTypeActions", "AssetTypeActions_AnimMontage", "Animation Montage"); }
+	virtual FLinearColor GetAssetColor() const override { return FLinearColor(FColor(100, 100, 255)); }
+	virtual TSoftClassPtr<UObject> GetAssetClass() const override { return UAnimMontage::StaticClass(); }
+	// UAssetDefinition End
+};
diff --git a/Engine/Plugins/Editor/EngineAssetDefinitions/Source/Private/Animation/AssetDefinition_AnimSequence.cpp b/Engine/Plugins/Editor/EngineAssetDefinitions/Source/Private/Animation/AssetDefinition_AnimSequence.cpp
new file mode 100644
--- /dev/null
+++ b/Engine/Plugins/Editor/EngineAssetDefinitions/Source/Private/Animation/AssetDefinition_AnimSequence.cpp
@@ -0,0 +1,201 @@
+// Copyright Epic Games, Inc. All Rights Reserved.
+
+#include "AssetDefinition_AnimSequence.h"
+#include "Animation/AnimSequence.h"
+#include "ToolMenus.h"
+#include "Styling/AppStyle.h"
+#include "EditorReimportHandler.h"
+#include "Animation/AnimMontage.h"
+#include "Factories/AnimCompositeFactory.h"
+#include "Factories/AnimStreamableFactory.h"
+#include "Factories/AnimMontageFactory.h"
+#include "Factories/PoseAssetFactory.h"
+#include "Animation/AnimComposite.h"
+#include "Animation/AnimStreamable.h"
+#include "Animation/PoseAsset.h"
+#include "IAssetTools.h"
+#include "Algo/Transform.h"
+#include "ToolMenu.h"
+#include "ToolMenus.h"
+#include "ContentBrowserMenuContexts.h"
+#include "Misc/App.h"
+
+#define LOCTEXT_NAMESPACE "AssetTypeActions"
+
+// Menu Extensions
+//--------------------------------------------------------------------
+
+namespace MenuExtension_AnimSequence
+{
+	/** Delegate used when creating Assets from an AnimSequence */
+	DECLARE_DELEGATE_RetVal_TwoParams(bool, FOnConfigureFactory, class UFactory*, class UAnimSequence* );
+
+	int32 GEnableAnimStreamable = 0;
+	static const TCHAR* AnimStreamableCVarName = TEXT("a.EnableAnimStreamable");
+
+	static FAutoConsoleVariableRef CVarEnableAnimStreamable(
+		AnimStreamableCVarName,
+		GEnableAnimStreamable,
+		TEXT("1 = Enables ability to make Anim Streamable assets. 0 = off"));
+	
+	void ExecuteNewAnimComposite(const FToolMenuContext& InContext)
+	{
+		const UContentBrowserAssetContextMenuContext* CBContext = UContentBrowserAssetContextMenuContext::FindContextWithAssets(InContext);
+
+		IAssetTools::Get().CreateAssetsFrom<UAnimSequence>(
+			CBContext->LoadSelectedObjects<UAnimSequence>(), UAnimComposite::StaticClass(), TEXT("_Composite"), [](UAnimSequence* SourceObject)
+			{
+				UAnimCompositeFactory* Factory = NewObject<UAnimCompositeFactory>();
+				Factory->SourceAnimation = SourceObject;
+				return Factory;
+			}
+		);
+	}
+
+	void ExecuteNewAnimMontage(const FToolMenuContext& InContext)
+	{
+		const UContentBrowserAssetContextMenuContext* CBContext = UContentBrowserAssetContextMenuContext::FindContextWithAssets(InContext);
+
+		IAssetTools::Get().CreateAssetsFrom<UAnimSequence>(
+			CBContext->LoadSelectedObjects<UAnimSequence>(), UAnimMontage::StaticClass(), TEXT("_Montage"), [](UAnimSequence* SourceObject)
+			{
+				UAnimMontageFactory* Factory = NewObject<UAnimMontageFactory>();
+				Factory->SourceAnimation = SourceObject;
+				return Factory;
+			}
+		);
+	}
+
+	void ExecuteNewAnimStreamable(const FToolMenuContext& InContext)
+	{
+		const UContentBrowserAssetContextMenuContext* CBContext = UContentBrowserAssetContextMenuContext::FindContextWithAssets(InContext);
+
+		IAssetTools::Get().CreateAssetsFrom<UAnimSequence>(
+			CBContext->LoadSelectedObjects<UAnimSequence>(), UAnimMontage::StaticClass(), TEXT("_Streamable"), [](UAnimSequence* SourceObject)
+			{
+				UAnimStreamableFactory* Factory = NewObject<UAnimStreamableFactory>();
+				Factory->SourceAnimation = SourceObject;
+				return Factory;
+			}
+		);
+	}
+
+	void ExecuteNewPoseAsset(const FToolMenuContext& InContext)
+	{
+		const UContentBrowserAssetContextMenuContext* CBContext = UContentBrowserAssetContextMenuContext::FindContextWithAssets(InContext);
+
+		IAssetTools::Get().CreateAssetsFrom<UAnimSequence>(
+			CBContext->LoadSelectedObjects<UAnimSequence>(), UPoseAsset::StaticClass(), TEXT("_PoseAsset"), [](UAnimSequence* SourceObject)
+			{
+				UPoseAssetFactory* Factory = NewObject<UPoseAssetFactory>();
+				Factory->SourceAnimation = SourceObject;
+				return Factory->ConfigureProperties() ? Factory : nullptr;
+			}
+		);
+	}
+	
+	void FillCreateMenu(UToolMenu* Menu)
+	{
+		IAssetTools& AssetTools = IAssetTools::Get();
+
+		FToolMenuSection& Section = Menu->FindOrAddSection("CreateAssetsMenu");
+
+		if (AssetTools.IsAssetClassSupported(UAnimComposite::StaticClass()))
+		{
+			const TAttribute<FText> Label = LOCTEXT("AnimSequence_NewAnimComposite", "Create AnimComposite");
+			const TAttribute<FText> ToolTip = LOCTEXT("AnimSequence_NewAnimCompositeTooltip", "Creates an AnimComposite using the selected anim sequence.");
+			const FSlateIcon Icon = FSlateIcon(FAppStyle::GetAppStyleSetName(), "ClassIcon.AnimComposite");
+
+			FToolUIAction UIAction;
+			UIAction.ExecuteAction = FToolMenuExecuteAction::CreateStatic(&ExecuteNewAnimComposite);
+			Section.AddMenuEntry("AnimSequence_NewAnimComposite", Label, ToolTip, Icon, UIAction);
+		}
+
+		if (AssetTools.IsAssetClassSupported(UAnimMontage::StaticClass()))
+		{
+			const TAttribute<FText> Label = LOCTEXT("AnimSequence_NewAnimMontage", "Create AnimMontage");
+			const TAttribute<FText> ToolTip = LOCTEXT("AnimSequence_NewAnimMontageTooltip", "Creates an AnimMontage using the selected anim sequence.");
+			const FSlateIcon Icon = FSlateIcon(FAppStyle::GetAppStyleSetName(), "ClassIcon.AnimMontage");
+
+			FToolUIAction UIAction;
+			UIAction.ExecuteAction = FToolMenuExecuteAction::CreateStatic(&ExecuteNewAnimMontage);
+			Section.AddMenuEntry("AnimSequence_NewAnimMontage", Label, ToolTip, Icon, UIAction);
+		}
+
+		if (GEnableAnimStreamable == 1)
+		{
+			if (AssetTools.IsAssetClassSupported(UAnimStreamable::StaticClass()))
+			{
+				const TAttribute<FText> Label = LOCTEXT("AnimSequence_NewAnimStreamable", "Create AnimStreamable");
+				const TAttribute<FText> ToolTip = LOCTEXT("AnimSequence_NewAnimStreamableTooltip", "Creates an AnimStreamable using the selected anim sequence.");
+				const FSlateIcon Icon = FSlateIcon(FAppStyle::GetAppStyleSetName(), "ClassIcon.AnimMontage");
+
+				FToolUIAction UIAction;
+				UIAction.ExecuteAction = FToolMenuExecuteAction::CreateStatic(&ExecuteNewAnimStreamable);
+				Section.AddMenuEntry("AnimSequence_NewAnimStreamable", Label, ToolTip, Icon, UIAction);
+			}
+		}
+
+		if (AssetTools.IsAssetClassSupported(UPoseAsset::StaticClass()))
+		{
+			const TAttribute<FText> Label = LOCTEXT("AnimSequence_NewPoseAsset", "Create PoseAsset");
+			const TAttribute<FText> ToolTip = LOCTEXT("AnimSequence_NewPoseAssetTooltip", "Creates an PoseAsset using the selected anim sequence.");
+			const FSlateIcon Icon = FSlateIcon(FAppStyle::GetAppStyleSetName(), "ClassIcon.PoseAsset");
+
+			FToolUIAction UIAction;
+			UIAction.ExecuteAction = FToolMenuExecuteAction::CreateStatic(&ExecuteNewPoseAsset);
+			Section.AddMenuEntry("AnimSequence_NewPoseAsset", Label, ToolTip, Icon, UIAction);
+		}
+	}
+
+	void ExecuteReimportWithNewSource(const FToolMenuContext& InContext)
+	{
+		const UContentBrowserAssetContextMenuContext* CBContext = UContentBrowserAssetContextMenuContext::FindContextWithAssets(InContext);
+		TArray<UObject*> ReimportAssets = CBContext->LoadSelectedObjects<UObject>();
+	
+		const bool bShowNotification = !FApp::IsUnattended();
+		const bool bReimportWithNewFile = true;
+		const int32 SourceFileIndex = INDEX_NONE;
+		FReimportManager::Instance()->ValidateAllSourceFileAndReimport(ReimportAssets, bShowNotification, SourceFileIndex, bReimportWithNewFile);
+	}
+
+	static FDelayedAutoRegisterHelper DelayedAutoRegister(EDelayedRegisterRunPhase::EndOfEngineInit, []{ 
+		UToolMenus::RegisterStartupCallback(FSimpleMulticastDelegate::FDelegate::CreateLambda([]()
+		{
+			FToolMenuOwnerScoped OwnerScoped(UE_MODULE_NAME);
+			UToolMenu* Menu = UE::ContentBrowser::ExtendToolMenu_AssetContextMenu(UAnimSequence::StaticClass());
+		
+			FToolMenuSection& Section = Menu->FindOrAddSection("GetAssetActions");
+			Section.AddDynamicEntry(NAME_None, FNewToolMenuSectionDelegate::CreateLambda([](FToolMenuSection& InSection)
+			{
+				IAssetTools& AssetTools = IAssetTools::Get();
+
+				if (AssetTools.IsAssetClassSupported(UAnimComposite::StaticClass()) ||
+					AssetTools.IsAssetClassSupported(UAnimMontage::StaticClass()) ||
+					AssetTools.IsAssetClassSupported(UAnimStreamable::StaticClass()) ||
+					AssetTools.IsAssetClassSupported(UPoseAsset::StaticClass()))
+				{
+					InSection.AddSubMenu(
+						"AnimSequence_CreateAnimSubmenu",
+						LOCTEXT("CreateAnimSubmenu", "Create"),
+						LOCTEXT("CreateAnimSubmenu_ToolTip", "Create assets from this anim sequence"),
+						FNewToolMenuDelegate::CreateStatic(&FillCreateMenu),
+						false,
+						FSlateIcon(FAppStyle::GetAppStyleSetName(), "Persona.AssetActions.CreateAnimAsset")
+					);
+				}
+				{
+					const TAttribute<FText> Label = LOCTEXT("AnimSequence_ReimportWithNewSource", "Reimport with New Source");
+					const TAttribute<FText> ToolTip = LOCTEXT("AnimSequence_ReimportWithNewSourceTooltip", "Reimport the selected sequence(s) from a new source file.");
+					const FSlateIcon Icon = FSlateIcon(FAppStyle::GetAppStyleSetName(), "Persona.AssetActions.ReimportAnim");
+
+					FToolUIAction UIAction;
+					UIAction.ExecuteAction = FToolMenuExecuteAction::CreateStatic(&ExecuteReimportWithNewSource);
+					InSection.AddMenuEntry("AnimSequence_ReimportWithNewSource", Label, ToolTip, Icon, UIAction);
+				}
+			}));
+		}));
+	});
+}
+
+#undef LOCTEXT_NAMESPACE
diff --git a/Engine/Plugins/Editor/EngineAssetDefinitions/Source/Private/Animation/AssetDefinition_AnimSequence.h b/Engine/Plugins/Editor/EngineAssetDefinitions/Source/Private/Animation/AssetDefinition_AnimSequence.h
new file mode 100644
--- /dev/null
+++ b/Engine/Plugins/Editor/EngineAssetDefinitions/Source/Private/Animation/AssetDefinition_AnimSequence.h
@@ -0,0 +1,22 @@
+// Copyright Epic Games, Inc. All Rights Reserved.
+
+#pragma once
+
+#include "Animation/AnimSequence.h"
+#include "AssetDefinition_AnimationAsset.h"
+
+#include "AssetDefinition_AnimSequence.generated.h"
+
+UCLASS()
+class UAssetDefinition_AnimSequence : public UAssetDefinition_AnimationAsset
+{
+	GENERATED_BODY()
+
+public:
+	// UAssetDefinition Begin
+	virtual FText GetAssetDisplayName() const override { return NSLOCTEXT("AssetTypeActions", "AssetTypeActions_AnimSequence", "Animation Sequence"); }
+	virtual FLinearColor GetAssetColor() const override { return FLinearColor(FColor(237, 28, 36)); }
+	virtual TSoftClassPtr<UObject> GetAssetClass() const override { return UAnimSequence::StaticClass(); }
+	virtual bool CanImport() const override { return true; }
+	// UAssetDefinition End
+};
diff --git a/Engine/Plugins/Editor/EngineAssetDefinitions/Source/Private/Animation/AssetDefinition_AnimStreamable.h b/Engine/Plugins/Editor/EngineAssetDefinitions/Source/Private/Animation/AssetDefinition_AnimStreamable.h
new file mode 100644
--- /dev/null
+++ b/Engine/Plugins/Editor/EngineAssetDefinitions/Source/Private/Animation/AssetDefinition_AnimStreamable.h
@@ -0,0 +1,21 @@
+// Copyright Epic Games, Inc. All Rights Reserved.
+
+#pragma once
+
+#include "Animation/AnimStreamable.h"
+#include "AssetDefinition_AnimationAsset.h"
+
+#include "AssetDefinition_AnimStreamable.generated.h"
+
+UCLASS()
+class UAssetDefinition_AnimStreamable : public UAssetDefinition_AnimationAsset
+{
+	GENERATED_BODY()
+
+public:
+	// UAssetDefinition Begin
+	virtual FText GetAssetDisplayName() const override { return NSLOCTEXT("AssetTypeActions", "AssetTypeActions_AnimStreamable", "Streamable Animation"); }
+	virtual FLinearColor GetAssetColor() const override { return FLinearColor(FColor(181, 230, 29)); }
+	virtual TSoftClassPtr<UObject> GetAssetClass() const override { return UAnimStreamable::StaticClass(); }
+	// UAssetDefinition End
+};
diff --git a/Engine/Plugins/Editor/EngineAssetDefinitions/Source/Private/Animation/AssetDefinition_AnimationAsset.cpp b/Engine/Plugins/Editor/EngineAssetDefinitions/Source/Private/Animation/AssetDefinition_AnimationAsset.cpp
new file mode 100644
--- /dev/null
+++ b/Engine/Plugins/Editor/EngineAssetDefinitions/Source/Private/Animation/AssetDefinition_AnimationAsset.cpp
@@ -0,0 +1,243 @@
+// Copyright Epic Games, Inc. All Rights Reserved.
+
+#include "AssetDefinition_AnimationAsset.h"
+#include "Framework/MultiBox/MultiBoxBuilder.h"
+#include "Misc/MessageDialog.h"
+#include "Styling/AppStyle.h"
+#include "ThumbnailRendering/SceneThumbnailInfo.h"
+#include "IAssetTools.h"
+#include "ContentBrowserMenuContexts.h"
+#include "SSkeletonWidget.h"
+#include "IAnimationEditorModule.h"
+#include "Preferences/PersonaOptions.h"
+#include "ToolMenus.h"
+#include "Framework/Application/SlateApplication.h"
+#include "Algo/Transform.h"
+#include "IPersonaToolkit.h"
+#include "Subsystems/AssetEditorSubsystem.h"
+#include "Editor.h"
+
+#define LOCTEXT_NAMESPACE "AssetTypeActions"
+
+namespace UE::AnimationAsset
+{
+	bool ReplaceMissingSkeleton(const TArray<TWeakObjectPtr<UAnimationAsset>>& InAnimationAssets)
+	{
+		// record anim assets that need skeleton replaced
+		const TArray<TWeakObjectPtr<UObject>> AnimsToFix = TArray<TWeakObjectPtr<UObject>>(InAnimationAssets);
+		
+		// get a skeleton from the user and replace it
+		const TSharedPtr<SReplaceMissingSkeletonDialog> PickSkeletonWindow = SNew(SReplaceMissingSkeletonDialog).AnimAssets(AnimsToFix);
+		const bool bWasSkeletonReplaced = PickSkeletonWindow.Get()->ShowModal();
+		return bWasSkeletonReplaced;
+	}
+	
+	void OpenAnimAssetEditor(const TArrayView<UAnimationAsset*>& InAnimAssets, bool bForceNewEditor, EToolkitMode::Type Mode, TSharedPtr<IToolkitHost> ToolkitHost)
+	{
+#if WITH_EDITOR
+		// Force opening a new tab when the Animation Editor options window has the 'Always Open In New Tab' option enabled.
+		const UPersonaOptions* PersonaOptions = GetDefault<UPersonaOptions>();
+		check(PersonaOptions);
+		bForceNewEditor |= PersonaOptions->bAlwaysOpenAnimationAssetsInNewTab;
+#endif
+
+		// Force new tab when shift is held down.
+		bForceNewEditor |= FSlateApplication::Get().GetModifierKeys().IsShiftDown();
+
+		// For each one..
+		for (UAnimationAsset* AnimAsset : InAnimAssets)
+		{
+			USkeleton* AnimSkeleton = AnimAsset->GetSkeleton();
+			if (!AnimSkeleton)
+			{
+				FText ShouldRetargetMessage = LOCTEXT("ShouldRetargetAnimAsset_Message", "Could not find the skeleton for Anim '{AnimName}' Would you like to choose a new one?");
+
+				FFormatNamedArguments Arguments;
+				Arguments.Add(TEXT("AnimName"), FText::FromString(AnimAsset->GetName()));
+
+				if (FMessageDialog::Open(EAppMsgType::YesNo, FText::Format(ShouldRetargetMessage, Arguments)) == EAppReturnType::Yes)
+				{
+					TArray<TWeakObjectPtr<UAnimationAsset>> AssetsToRetarget;
+					AssetsToRetarget.Add(AnimAsset);
+					const bool bSkeletonReplaced = ReplaceMissingSkeleton(AssetsToRetarget);
+					if (!bSkeletonReplaced)
+					{
+						return; // Persona will crash if trying to load asset without a skeleton
+					}
+				}
+				else
+				{
+					return;
+				}
+			}
+		
+			// First see if we already have it open
+			const bool bBringToFrontIfOpen = true;
+#if WITH_EDITOR
+			if (IAssetEditorInstance* EditorInstance = GEditor->GetEditorSubsystem<UAssetEditorSubsystem>()->FindEditorForAsset(AnimAsset, bBringToFrontIfOpen))
+			{
+				EditorInstance->FocusWindow();
+			}
+			else
+#endif
+			{
+				// See if we are trying to open a single asset. If we are, we re-use a compatible anim editor.
+				bool bSingleAsset = InAnimAssets.Num() == 1;
+				bool bFoundEditor = false;
+				if (bSingleAsset && !bForceNewEditor)
+				{
+					// See if there is an animation asset with the same skeleton already being edited
+					TArray<UObject*> AllEditedAssets;
+#if WITH_EDITOR
+					AllEditedAssets = GEditor->GetEditorSubsystem<UAssetEditorSubsystem>()->GetAllEditedAssets();
+#endif
+					UAnimationAsset* CompatibleEditedAsset = nullptr;
+					for (UObject* EditedAsset : AllEditedAssets)
+					{
+						UAnimationAsset* EditedAnimAsset = Cast<UAnimationAsset>(EditedAsset);
+						if (EditedAnimAsset && EditedAnimAsset->GetSkeleton() == AnimAsset->GetSkeleton())
+						{
+							CompatibleEditedAsset = EditedAnimAsset;
+							break;
+						}
+					}
+
+					// If there is..
+					if(CompatibleEditedAsset)
+					{
+						// Find the anim editors that are doing it
+						TArray<IAssetEditorInstance*> AssetEditors;
+#if WITH_EDITOR
+						AssetEditors = GEditor->GetEditorSubsystem<UAssetEditorSubsystem>()->FindEditorsForAsset(CompatibleEditedAsset);
+#endif
+						for (IAssetEditorInstance* ExistingEditor : AssetEditors)
+						{
+							if (ExistingEditor->GetEditorName() == FName("AnimationEditor"))
+							{
+								// Change the current anim to this one
+								IAnimationEditor* AnimEditor = static_cast<IAnimationEditor*>(ExistingEditor);
+								if(AnimEditor->GetPersonaToolkit()->GetSkeleton() == AnimAsset->GetSkeleton())
+								{
+									AnimEditor->SetAnimationAsset(AnimAsset);
+									AnimEditor->FocusWindow();
+									bFoundEditor = true;
+									break;
+								}
+							}
+						}
+					}
+				}
+
+				// We didn't find an editor, make a new one
+				if (!bFoundEditor)
+				{
+					IAnimationEditorModule& AnimationEditorModule = FModuleManager::LoadModuleChecked<IAnimationEditorModule>("AnimationEditor");
+					AnimationEditorModule.CreateAnimationEditor(Mode, ToolkitHost, AnimAsset);
+				}
+			}
+		}
+	}
+}
+
+EAssetCommandResult UAssetDefinition_AnimationAsset::OpenAssets(const FAssetOpenArgs& OpenArgs) const
+{
+	UE::AnimationAsset::OpenAnimAssetEditor(OpenArgs.LoadObjects<UAnimationAsset>(), false, OpenArgs.GetToolkitMode(), OpenArgs.ToolkitHost);
+
+	return EAssetCommandResult::Handled;
+}
+
+UThumbnailInfo* UAssetDefinition_AnimationAsset::LoadThumbnailInfo(const FAssetData& InAssetData) const
+{
+	UAnimationAsset* Anim = CastChecked<UAnimationAsset>(InAssetData.GetAsset());
+	UThumbnailInfo* ThumbnailInfo = Anim->ThumbnailInfo;
+	if (ThumbnailInfo == nullptr)
+	{
+		ThumbnailInfo = NewObject<USceneThumbnailInfo>(Anim, NAME_None, RF_Transactional);
+		Anim->ThumbnailInfo = ThumbnailInfo;
+	}
+
+	return ThumbnailInfo;
+}
+
+// Menu Extensions
+//--------------------------------------------------------------------
+
+namespace MenuExtension_AnimationAsset
+{
+	void ExecuteOpenInNewWindow(const FToolMenuContext& InContext)
+	{
+		const UContentBrowserAssetContextMenuContext* Context = UContentBrowserAssetContextMenuContext::FindContextWithAssets(InContext);
+		TArray<UAnimationAsset*> AnimationAssets = Context->LoadSelectedObjects<UAnimationAsset>();
+		
+		UE::AnimationAsset::OpenAnimAssetEditor(AnimationAssets, true, EToolkitMode::Standalone, nullptr);
+	}
+
+	void ExecuteFindSkeleton(const FToolMenuContext& InContext)
+	{
+		const UContentBrowserAssetContextMenuContext* Context = UContentBrowserAssetContextMenuContext::FindContextWithAssets(InContext);
+
+		TArray<UObject*> ObjectsToSync;
+		for (UAnimationAsset* AnimationAsset : Context->LoadSelectedObjects<UAnimationAsset>())
+		{
+			if (USkeleton* Skeleton = AnimationAsset->GetSkeleton())
+			{
+				ObjectsToSync.AddUnique(Skeleton);
+			}
+		}
+		
+		if ( ObjectsToSync.Num() > 0 )
+		{
+			IAssetTools::Get().SyncBrowserToAssets(ObjectsToSync);
+		}
+	}
+
+	void ExecuteReplaceSkeleton(const FToolMenuContext& InContext)
+	{
+		const UContentBrowserAssetContextMenuContext* Context = UContentBrowserAssetContextMenuContext::FindContextWithAssets(InContext);
+		TArray<UAnimationAsset*> AnimationAssets = Context->LoadSelectedObjects<UAnimationAsset>();
+		
+		UE::AnimationAsset::ReplaceMissingSkeleton(TArray<TWeakObjectPtr<UAnimationAsset>>(AnimationAssets));
+	}
+	
+	static FDelayedAutoRegisterHelper DelayedAutoRegister(EDelayedRegisterRunPhase::EndOfEngineInit, []{ 
+		UToolMenus::RegisterStartupCallback(FSimpleMulticastDelegate::FDelegate::CreateLambda([]()
+		{
+			FToolMenuOwnerScoped OwnerScoped(UE_MODULE_NAME);
+			UToolMenu* Menu = UE::ContentBrowser::ExtendToolMenu_AssetContextMenu(UAnimationAsset::StaticClass());
+		
+			FToolMenuSection& Section = Menu->FindOrAddSection("GetAssetActions");
+			Section.AddDynamicEntry(NAME_None, FNewToolMenuSectionDelegate::CreateLambda([](FToolMenuSection& InSection)
+			{
+				{
+					const TAttribute<FText> Label = LOCTEXT("AnimSequenceBase_OpenInNewWindow", "Open In New Window");
+					const TAttribute<FText> ToolTip = LOCTEXT("AnimSequenceBase_OpenInNewWindowTooltip", "Will always open asset in a new window, and not re-use existing window. (Shift+Double-Click)");
+					const FSlateIcon Icon = FSlateIcon(FAppStyle::GetAppStyleSetName(), "Icons.OpenInExternalEditor");
+
+					FToolUIAction UIAction;
+					UIAction.ExecuteAction = FToolMenuExecuteAction::CreateStatic(&ExecuteOpenInNewWindow);
+					InSection.AddMenuEntry("AnimSequenceBase_OpenInNewWindow", Label, ToolTip, Icon, UIAction);
+				}
+				{
+					const TAttribute<FText> Label = LOCTEXT("AnimSequenceBase_FindSkeleton", "Find Skeleton");
+					const TAttribute<FText> ToolTip = LOCTEXT("AnimSequenceBase_FindSkeletonTooltip", "Finds the skeleton for the selected assets in the content browser.");
+					const FSlateIcon Icon = FSlateIcon(FAppStyle::GetAppStyleSetName(), "Icons.Find");
+
+					FToolUIAction UIAction;
+					UIAction.ExecuteAction = FToolMenuExecuteAction::CreateStatic(&ExecuteFindSkeleton);
+					InSection.AddMenuEntry("AnimSequenceBase_FindSkeleton", Label, ToolTip, Icon, UIAction);
+				}
+				{
+					const TAttribute<FText> Label = LOCTEXT("AnimSequenceBase_ReplaceSkeleton", "Replace Skeleton");
+					const TAttribute<FText> ToolTip = LOCTEXT("AnimSequenceBase_ReplaceSkeletonTooltip", "Associate a different skeleton with the selected animation assets in the content browser.");
+					const FSlateIcon Icon = FSlateIcon(FAppStyle::GetAppStyleSetName(), "Persona.RetargetManager");
+
+					FToolUIAction UIAction;
+					UIAction.ExecuteAction = FToolMenuExecuteAction::CreateStatic(&ExecuteReplaceSkeleton);
+					InSection.AddMenuEntry("AnimSequenceBase_ReplaceSkeleton", Label, ToolTip, Icon, UIAction);
+				}
+			}));
+		}));
+	});
+}
+
+#undef LOCTEXT_NAMESPACE
diff --git a/Engine/Plugins/Editor/EngineAssetDefinitions/Source/Private/Animation/AssetDefinition_AnimationAsset.h b/Engine/Plugins/Editor/EngineAssetDefinitions/Source/Private/Animation/AssetDefinition_AnimationAsset.h
new file mode 100644
--- /dev/null
+++ b/Engine/Plugins/Editor/EngineAssetDefinitions/Source/Private/Animation/AssetDefinition_AnimationAsset.h
@@ -0,0 +1,28 @@
+// Copyright Epic Games, Inc. All Rights Reserved.
+
+#pragma once
+
+#include "Animation/AnimationAsset.h"
+#include "AssetDefinitionDefault.h"
+
+#include "AssetDefinition_AnimationAsset.generated.h"
+
+UCLASS()
+class UAssetDefinition_AnimationAsset : public UAssetDefinitionDefault
+{
+	GENERATED_BODY()
+
+public:
+	// UAssetDefinition Begin
+	virtual FText GetAssetDisplayName() const override { return NSLOCTEXT("AssetTypeActions", "AssetTypeActions_AnimationAsset", "AnimationAsset"); }
+	virtual FLinearColor GetAssetColor() const override { return FLinearColor(FColor(80, 123, 72)); }
+	virtual TSoftClassPtr<UObject> GetAssetClass() const override { return UAnimationAsset::StaticClass(); }
+	virtual TConstArrayView<FAssetCategoryPath> GetAssetCategories() const override
+	{
+		static const auto Categories = { EAssetCategoryPaths::Animation };
+		return Categories;
+	}
+	virtual UThumbnailInfo* LoadThumbnailInfo(const FAssetData& InAssetData) const override;
+	virtual EAssetCommandResult OpenAssets(const FAssetOpenArgs& OpenArgs) const override;
+	// UAssetDefinition End
+};
diff --git a/Engine/Plugins/Editor/EngineAssetDefinitions/Source/Private/Animation/AssetDefinition_BlendSpace.h b/Engine/Plugins/Editor/EngineAssetDefinitions/Source/Private/Animation/AssetDefinition_BlendSpace.h
new file mode 100644
--- /dev/null
+++ b/Engine/Plugins/Editor/EngineAssetDefinitions/Source/Private/Animation/AssetDefinition_BlendSpace.h
@@ -0,0 +1,21 @@
+// Copyright Epic Games, Inc. All Rights Reserved.
+
+#pragma once
+
+#include "Animation/BlendSpace.h"
+#include "AssetDefinition_AnimationAsset.h"
+
+#include "AssetDefinition_BlendSpace.generated.h"
+
+UCLASS()
+class UAssetDefinition_BlendSpace : public UAssetDefinition_AnimationAsset
+{
+	GENERATED_BODY()
+
+public:
+	// UAssetDefinition Begin
+	virtual FText GetAssetDisplayName() const override { return NSLOCTEXT("AssetTypeActions", "AssetTypeActions_BlendSpace", "Blend Space"); }
+	virtual FLinearColor GetAssetColor() const override { return FLinearColor(FColor(255, 168, 111)); }
+	virtual TSoftClassPtr<UObject> GetAssetClass() const override { return UBlendSpace::StaticClass(); }
+	// UAssetDefinition End
+};
\ No newline at end of file
diff --git a/Engine/Plugins/Editor/EngineAssetDefinitions/Source/Private/Animation/AssetDefinition_BlendSpace1D.h b/Engine/Plugins/Editor/EngineAssetDefinitions/Source/Private/Animation/AssetDefinition_BlendSpace1D.h
new file mode 100644
--- /dev/null
+++ b/Engine/Plugins/Editor/EngineAssetDefinitions/Source/Private/Animation/AssetDefinition_BlendSpace1D.h
@@ -0,0 +1,26 @@
+// Copyright Epic Games, Inc. All Rights Reserved.
+
+#pragma once
+
+#include "Animation/BlendSpace1D.h"
+#include "AssetDefinition_AnimationAsset.h"
+
+#include "AssetDefinition_BlendSpace1D.generated.h"
+
+UCLASS()
+class UAssetDefinition_BlendSpace1D : public UAssetDefinition_AnimationAsset
+{
+	GENERATED_BODY()
+
+public:
+	// UAssetDefinition Begin
+	virtual FText GetAssetDisplayName() const override { return NSLOCTEXT("AssetTypeActions", "AssetTypeActions_BlendSpace1D", "Blend Space 1D"); }
+	virtual FLinearColor GetAssetColor() const override { return FLinearColor(FColor(255, 180, 130)); }
+	virtual TSoftClassPtr<UObject> GetAssetClass() const override { return UBlendSpace1D::StaticClass(); }
+	virtual TConstArrayView<FAssetCategoryPath> GetAssetCategories() const override
+	{
+		static const auto Categories = { EAssetCategoryPaths::Animation / NSLOCTEXT("AssetTypeActions", "AnimLegacySubMenu", "Legacy") };
+		return Categories;
+	}
+	// UAssetDefinition End
+};
\ No newline at end of file
diff --git a/Engine/Plugins/Editor/EngineAssetDefinitions/Source/Private/Animation/AssetDefinition_PoseAsset.cpp b/Engine/Plugins/Editor/EngineAssetDefinitions/Source/Private/Animation/AssetDefinition_PoseAsset.cpp
new file mode 100644
--- /dev/null
+++ b/Engine/Plugins/Editor/EngineAssetDefinitions/Source/Private/Animation/AssetDefinition_PoseAsset.cpp
@@ -0,0 +1,65 @@
+// Copyright Epic Games, Inc. All Rights Reserved.
+
+#include "AssetDefinition_PoseAsset.h"
+#include "Animation/Skeleton.h"
+#include "ToolMenuSection.h"
+#include "Animation/AnimSequence.h"
+#include "ScopedTransaction.h"
+#include "ToolMenu.h"
+#include "ToolMenus.h"
+#include "ContentBrowserMenuContexts.h"
+
+#define LOCTEXT_NAMESPACE "AssetTypeActions"
+
+// Menu Extensions
+//--------------------------------------------------------------------
+
+namespace MenuExtension_PoseAsset
+{
+	void ExecuteUpdateSource(const FToolMenuContext& InContext)
+	{
+		const UContentBrowserAssetContextMenuContext* CBContext = UContentBrowserAssetContextMenuContext::FindContextWithAssets(InContext);
+
+		FScopedTransaction Transaction(LOCTEXT("PoseUpdateSource", "Updating Source Animation for Pose"));
+		for (UPoseAsset* PoseAsset : CBContext->LoadSelectedObjects<UPoseAsset>())
+		{
+			if (UAnimSequence* SourceAnimation = PoseAsset->SourceAnimation)
+			{
+				if (!PoseAsset->SourceAnimationRawDataGUID.IsValid() || PoseAsset->SourceAnimationRawDataGUID != SourceAnimation->GetRawDataGuid())
+				{
+					if (PoseAsset->GetSkeleton()->IsCompatibleForEditor(SourceAnimation->GetSkeleton()))
+					{
+						PoseAsset->Modify();
+						PoseAsset->UpdatePoseFromAnimation(SourceAnimation);
+					}
+				}
+			}
+		}
+	}
+	
+	static FDelayedAutoRegisterHelper DelayedAutoRegister(EDelayedRegisterRunPhase::EndOfEngineInit, []{ 
+		UToolMenus::RegisterStartupCallback(FSimpleMulticastDelegate::FDelegate::CreateLambda([]()
+		{
+			FToolMenuOwnerScoped OwnerScoped(UE_MODULE_NAME);
+			UToolMenu* Menu = UE::ContentBrowser::ExtendToolMenu_AssetContextMenu(UPoseAsset::StaticClass());
+		
+			FToolMenuSection& Section = Menu->FindOrAddSection("GetAssetActions");
+			Section.AddDynamicEntry(NAME_None, FNewToolMenuSectionDelegate::CreateLambda([](FToolMenuSection& InSection)
+			{
+				const UContentBrowserAssetContextMenuContext* CBContext = UContentBrowserAssetContextMenuContext::FindContextWithAssets(InSection);
+
+				{
+					const TAttribute<FText> Label = LOCTEXT("PoseAsset_UpdateSource", "Update Source Animation");
+					const TAttribute<FText> ToolTip = LOCTEXT("PoseAsset_UpdateSourceTooltip", "Updates the source animation for this pose");
+					const FSlateIcon Icon = FSlateIcon(FAppStyle::GetAppStyleSetName(), "ClassIcon.AnimMontage");
+
+					FToolUIAction UIAction;
+					UIAction.ExecuteAction = FToolMenuExecuteAction::CreateStatic(&ExecuteUpdateSource);
+					InSection.AddMenuEntry("PoseAsset_UpdateSource", Label, ToolTip, Icon, UIAction);
+				}
+			}));
+		}));
+	});
+}
+
+#undef LOCTEXT_NAMESPACE
diff --git a/Engine/Plugins/Editor/EngineAssetDefinitions/Source/Private/Animation/AssetDefinition_PoseAsset.h b/Engine/Plugins/Editor/EngineAssetDefinitions/Source/Private/Animation/AssetDefinition_PoseAsset.h
new file mode 100644
--- /dev/null
+++ b/Engine/Plugins/Editor/EngineAssetDefinitions/Source/Private/Animation/AssetDefinition_PoseAsset.h
@@ -0,0 +1,21 @@
+// Copyright Epic Games, Inc. All Rights Reserved.
+
+#pragma once
+
+#include "Animation/PoseAsset.h"
+#include "AssetDefinition_AnimationAsset.h"
+
+#include "AssetDefinition_PoseAsset.generated.h"
+
+UCLASS()
+class UAssetDefinition_PoseAsset : public UAssetDefinition_AnimationAsset
+{
+	GENERATED_BODY()
+
+public:
+	// UAssetDefinition Begin
+	virtual FText GetAssetDisplayName() const override { return NSLOCTEXT("AssetTypeActions", "AssetTypeActions_PoseAsset", "Pose Asset"); }
+	virtual FLinearColor GetAssetColor() const override { return FLinearColor(FColor(237, 28, 36)); }
+	virtual TSoftClassPtr<UObject> GetAssetClass() const override { return UPoseAsset::StaticClass(); }
+	// UAssetDefinition End
+};
diff --git a/Engine/Plugins/Editor/EngineAssetDefinitions/Source/Private/Animation/AssetDefinition_SkeletalMesh.cpp b/Engine/Plugins/Editor/EngineAssetDefinitions/Source/Private/Animation/AssetDefinition_SkeletalMesh.cpp
--- a/Engine/Plugins/Editor/EngineAssetDefinitions/Source/Private/Animation/AssetDefinition_SkeletalMesh.cpp
+++ b/Engine/Plugins/Editor/EngineAssetDefinitions/Source/Private/Animation/AssetDefinition_SkeletalMesh.cpp
@@ -363,8 +363,7 @@ namespace MenuExtension_SkeletalMesh
     
     	if ( ObjectsToSync.Num() > 0 )
     	{
-    		FContentBrowserModule& ContentBrowserModule = FModuleManager::LoadModuleChecked<FContentBrowserModule>("ContentBrowser");
-    		ContentBrowserModule.Get().SyncBrowserToAssets(ObjectsToSync);
+    		IAssetTools::Get().SyncBrowserToAssets(ObjectsToSync);
     	}
     }
     
diff --git a/Engine/Plugins/Editor/EngineAssetDefinitions/Source/Private/AssetDefinition_ActorFoliageSettings.h b/Engine/Plugins/Editor/EngineAssetDefinitions/Source/Private/AssetDefinition_ActorFoliageSettings.h
new file mode 100644
--- /dev/null
+++ b/Engine/Plugins/Editor/EngineAssetDefinitions/Source/Private/AssetDefinition_ActorFoliageSettings.h
@@ -0,0 +1,26 @@
+// Copyright Epic Games, Inc. All Rights Reserved.
+
+#pragma once
+
+#include "FoliageType_Actor.h"
+#include "AssetDefinitionDefault.h"
+
+#include "AssetDefinition_ActorFoliageSettings.generated.h"
+
+UCLASS()
+class UAssetDefinition_ActorFoliageSettings : public UAssetDefinitionDefault
+{
+	GENERATED_BODY()
+
+public:
+	// UAssetDefinition Begin
+	virtual FText GetAssetDisplayName() const override { return NSLOCTEXT("AssetTypeActions", "AssetTypeActions_ActorFoliageSettings", "Actor Foliage"); }
+	virtual FLinearColor GetAssetColor() const override { return FLinearColor(FColor(12, 65, 12)); }
+	virtual TSoftClassPtr<UObject> GetAssetClass() const override { return UFoliageType_Actor::StaticClass(); }
+	virtual TConstArrayView<FAssetCategoryPath> GetAssetCategories() const override
+	{
+		static const auto Categories = { EAssetCategoryPaths::Foliage };
+		return Categories;
+	}
+	// UAssetDefinition End
+};
diff --git a/Engine/Plugins/Editor/EngineAssetDefinitions/Source/Private/AssetDefinition_PoseAsset.cpp b/Engine/Plugins/Editor/EngineAssetDefinitions/Source/Private/AssetDefinition_PoseAsset.cpp
deleted file mode 100644
diff --git a/Engine/Plugins/Editor/EngineAssetDefinitions/Source/Private/AssetDefinition_PoseAsset.h b/Engine/Plugins/Editor/EngineAssetDefinitions/Source/Private/AssetDefinition_PoseAsset.h
deleted file mode 100644
diff --git a/Engine/Plugins/Editor/EngineAssetDefinitions/Source/Private/EngineAssetDefinitionsModule.cpp b/Engine/Plugins/Editor/EngineAssetDefinitions/Source/Private/EngineAssetDefinitionsModule.cpp
--- a/Engine/Plugins/Editor/EngineAssetDefinitions/Source/Private/EngineAssetDefinitionsModule.cpp
+++ b/Engine/Plugins/Editor/EngineAssetDefinitions/Source/Private/EngineAssetDefinitionsModule.cpp
@@ -1,6 +1,6 @@
 ﻿// Copyright Epic Games, Inc. All Rights Reserved.
 
-#include "AssetDefinition_PoseAsset.h"
+#include "Modules/ModuleManager.h"
 
 #define LOCTEXT_NAMESPACE "FEngineAssetDefinitionsModule"
 
@@ -11,33 +11,11 @@ public:
 	// IModuleInterface interface
 	virtual void StartupModule() override
 	{
-		// Register asset types... for now
-		IAssetTools& AssetTools = FModuleManager::LoadModuleChecked<FAssetToolsModule>("AssetTools").Get();
-		RegisterAssetTypeAction(AssetTools, MakeShareable(new FAssetTypeActions_PoseAsset));
 	}
 
 	virtual void ShutdownModule() override
 	{
-		// Unregister all the asset types that we registered
-		if (FModuleManager::Get().IsModuleLoaded("AssetTools"))
-		{
-			IAssetTools& AssetTools = FModuleManager::GetModuleChecked<FAssetToolsModule>("AssetTools").Get();
-			for (int32 Index = 0; Index < CreatedAssetTypeActions.Num(); ++Index)
-			{
-				AssetTools.UnregisterAssetTypeActions(CreatedAssetTypeActions[Index].ToSharedRef());
-			}
-		}
-		CreatedAssetTypeActions.Empty();
 	}
-
-private:
-	void RegisterAssetTypeAction(IAssetTools& AssetTools, TSharedRef<IAssetTypeActions> Action)
-	{
-		AssetTools.RegisterAssetTypeActions(Action);
-		CreatedAssetTypeActions.Add(Action);
-	}
-
-	TArray<TSharedPtr<IAssetTypeActions>> CreatedAssetTypeActions;
 };
 
 #undef LOCTEXT_NAMESPACE
diff --git a/Engine/Source/Developer/AssetTools/Private/AssetTools.cpp b/Engine/Source/Developer/AssetTools/Private/AssetTools.cpp
--- a/Engine/Source/Developer/AssetTools/Private/AssetTools.cpp
+++ b/Engine/Source/Developer/AssetTools/Private/AssetTools.cpp
@@ -42,19 +42,9 @@
 #include "AssetTypeActions/AssetTypeActions_AnimBlueprintInterface.h"
 #include "AssetTypeActions/AssetTypeActions_AnimBlueprintGeneratedClass.h"
 #include "AssetTypeActions/AssetTypeActions_AnimBoneCompressionSettings.h"
-#include "AssetTypeActions/AssetTypeActions_AnimComposite.h"
-#include "AssetTypeActions/AssetTypeActions_AnimStreamable.h"
 #include "AssetTypeActions/AssetTypeActions_AnimCurveCompressionSettings.h"
-#include "AssetTypeActions/AssetTypeActions_AnimMontage.h"
-#include "AssetTypeActions/AssetTypeActions_AnimSequence.h"
-#include "AssetTypeActions/AssetTypeActions_BlendSpace.h"
-#include "AssetTypeActions/AssetTypeActions_AimOffset.h"
-#include "AssetTypeActions/AssetTypeActions_BlendSpace1D.h"
-#include "AssetTypeActions/AssetTypeActions_AimOffset1D.h"
-#include "AssetTypeActions/AssetTypeActions_DataAsset.h"
 #include "AssetTypeActions/AssetTypeActions_Class.h"
 #include "AssetTypeActions/AssetTypeActions_ForceFeedbackEffect.h"
-#include "AssetTypeActions/AssetTypeActions_ActorFoliageSettings.h"
 #include "AssetTypeActions/AssetTypeActions_ParticleSystem.h"
 #include "AssetTypeActions/AssetTypeActions_PhysicalMaterialMask.h"
 #include "WorldPartition/WorldPartition.h"
@@ -190,7 +180,7 @@ public:
 			{
 				if (Category.GetCategory() != EAssetCategoryPaths::Misc.GetCategory())
 				{
-					EAssetTypeCategories::Type AdvancedCategoryBit = AssetTools.FindAdvancedAssetCategory(Category.GetCategory());
+					const EAssetTypeCategories::Type AdvancedCategoryBit = AssetTools.FindAdvancedAssetCategory(Category.GetCategory());
 					if (AdvancedCategoryBit == EAssetTypeCategories::Misc)
 					{
 						CategoryBits |= AssetTools.RegisterAdvancedAssetCategory(Category.GetCategory(), Category.GetCategoryText());
@@ -1336,28 +1326,18 @@ UAssetToolsImpl::UAssetToolsImpl(const FObjectInitializer& ObjectInitializer)
 	AllocatedCategoryBits.Add(TEXT("Media"), FAdvancedAssetCategory(EAssetTypeCategories::Media, LOCTEXT("MediaAssetCategory", "Media")));
 	AllocatedCategoryBits.Add(TEXT("Texture"), FAdvancedAssetCategory(EAssetTypeCategories::Textures, LOCTEXT("TextureAssetCategory", "Textures")));
 
-	EAssetTypeCategories::Type FoliageCategoryBit = RegisterAdvancedAssetCategory(FName(TEXT("Foliage")), LOCTEXT("FoliageAssetCategory", "Foliage"));
 	EAssetTypeCategories::Type InputCategoryBit = RegisterAdvancedAssetCategory(FName(TEXT("Input")), LOCTEXT("InputAssetsCategory", "Input"));
-
+	
 	// Register the built-in asset type actions
 	RegisterAssetTypeActions(MakeShareable(new FAssetTypeActions_AnimBlueprintInterface));
 	RegisterAssetTypeActions(MakeShareable(new FAssetTypeActions_AnimBlueprint));
 	RegisterAssetTypeActions(MakeShareable(new FAssetTypeActions_AnimBlueprintGeneratedClass));
 	RegisterAssetTypeActions(MakeShareable(new FAssetTypeActions_AnimBoneCompressionSettings));
-	RegisterAssetTypeActions(MakeShareable(new FAssetTypeActions_AnimComposite));
-	RegisterAssetTypeActions(MakeShareable(new FAssetTypeActions_AnimStreamable));
 	RegisterAssetTypeActions(MakeShareable(new FAssetTypeActions_AnimCurveCompressionSettings));
-	RegisterAssetTypeActions(MakeShareable(new FAssetTypeActions_AnimMontage));
-	RegisterAssetTypeActions(MakeShareable(new FAssetTypeActions_AnimSequence));
-	RegisterAssetTypeActions(MakeShareable(new FAssetTypeActions_AimOffset));
-	RegisterAssetTypeActions(MakeShareable(new FAssetTypeActions_AimOffset1D));
-	RegisterAssetTypeActions(MakeShareable(new FAssetTypeActions_BlendSpace));
-	RegisterAssetTypeActions(MakeShareable(new FAssetTypeActions_BlendSpace1D));
 	RegisterAssetTypeActions(MakeShareable(new FAssetTypeActions_Blueprint));
 	RegisterAssetTypeActions(MakeShareable(new FAssetTypeActions_BlueprintGeneratedClass));
 	RegisterAssetTypeActions(MakeShareable(new FAssetTypeActions_Class));
 	RegisterAssetTypeActions(MakeShareable(new FAssetTypeActions_ForceFeedbackEffect(InputCategoryBit)));
-	RegisterAssetTypeActions(MakeShareable(new FAssetTypeActions_ActorFoliageSettings(FoliageCategoryBit)));
 	RegisterAssetTypeActions(MakeShareable(new FAssetTypeActions_ParticleSystem));
 	RegisterAssetTypeActions(MakeShareable(new FAssetTypeActions_PhysicalMaterialMask));
 
diff --git a/Engine/Source/Developer/AssetTools/Private/AssetTypeActions/AssetTypeActions_ActorFoliageSettings.cpp b/Engine/Source/Developer/AssetTools/Private/AssetTypeActions/AssetTypeActions_ActorFoliageSettings.cpp
deleted file mode 100644
diff --git a/Engine/Source/Developer/AssetTools/Private/AssetTypeActions/AssetTypeActions_ActorFoliageSettings.h b/Engine/Source/Developer/AssetTools/Private/AssetTypeActions/AssetTypeActions_ActorFoliageSettings.h
deleted file mode 100644
diff --git a/Engine/Source/Developer/AssetTools/Private/AssetTypeActions/AssetTypeActions_AimOffset.h b/Engine/Source/Developer/AssetTools/Private/AssetTypeActions/AssetTypeActions_AimOffset.h
deleted file mode 100644
diff --git a/Engine/Source/Developer/AssetTools/Private/AssetTypeActions/AssetTypeActions_AimOffset1D.h b/Engine/Source/Developer/AssetTools/Private/AssetTypeActions/AssetTypeActions_AimOffset1D.h
deleted file mode 100644
diff --git a/Engine/Source/Developer/AssetTools/Private/AssetTypeActions/AssetTypeActions_AnimComposite.h b/Engine/Source/Developer/AssetTools/Private/AssetTypeActions/AssetTypeActions_AnimComposite.h
deleted file mode 100644
diff --git a/Engine/Source/Developer/AssetTools/Private/AssetTypeActions/AssetTypeActions_AnimMontage.cpp b/Engine/Source/Developer/AssetTools/Private/AssetTypeActions/AssetTypeActions_AnimMontage.cpp
deleted file mode 100644
diff --git a/Engine/Source/Developer/AssetTools/Private/AssetTypeActions/AssetTypeActions_AnimMontage.h b/Engine/Source/Developer/AssetTools/Private/AssetTypeActions/AssetTypeActions_AnimMontage.h
deleted file mode 100644
diff --git a/Engine/Source/Developer/AssetTools/Private/AssetTypeActions/AssetTypeActions_AnimSequence.cpp b/Engine/Source/Developer/AssetTools/Private/AssetTypeActions/AssetTypeActions_AnimSequence.cpp
deleted file mode 100644
diff --git a/Engine/Source/Developer/AssetTools/Private/AssetTypeActions/AssetTypeActions_AnimSequence.h b/Engine/Source/Developer/AssetTools/Private/AssetTypeActions/AssetTypeActions_AnimSequence.h
deleted file mode 100644
diff --git a/Engine/Source/Developer/AssetTools/Private/AssetTypeActions/AssetTypeActions_AnimStreamable.h b/Engine/Source/Developer/AssetTools/Private/AssetTypeActions/AssetTypeActions_AnimStreamable.h
deleted file mode 100644
diff --git a/Engine/Source/Developer/AssetTools/Private/AssetTypeActions/AssetTypeActions_AnimationAsset.cpp b/Engine/Source/Developer/AssetTools/Private/AssetTypeActions/AssetTypeActions_AnimationAsset.cpp
--- a/Engine/Source/Developer/AssetTools/Private/AssetTypeActions/AssetTypeActions_AnimationAsset.cpp
+++ b/Engine/Source/Developer/AssetTools/Private/AssetTypeActions/AssetTypeActions_AnimationAsset.cpp
@@ -23,43 +23,7 @@
 
 #define LOCTEXT_NAMESPACE "AssetTypeActions"
 
-void FAssetTypeActions_AnimationAsset::GetActions(const TArray<UObject*>& InObjects, FToolMenuSection& Section)
-{
-	auto AnimAssets = GetTypedWeakObjectPtrs<UAnimationAsset>(InObjects);
-
-	Section.AddMenuEntry(
-		"AnimSequenceBase_OpenInNewWindow",
-		LOCTEXT("AnimSequenceBase_OpenInNewWindow", "Open In New Window"),
-		LOCTEXT("AnimSequenceBase_OpenInNewWindowTooltip", "Will always open asset in a new window, and not re-use existing window. (Shift+Double-Click)"),
-		FSlateIcon(FAppStyle::GetAppStyleSetName(), "Icons.OpenInExternalEditor"),
-		FUIAction(
-			FExecuteAction::CreateSP(this, &FAssetTypeActions_AnimationAsset::ExecuteOpenInNewWindow, AnimAssets),
-			FCanExecuteAction()
-		)
-	);
-
-	Section.AddMenuEntry(
-		"AnimSequenceBase_FindSkeleton",
-		LOCTEXT("AnimSequenceBase_FindSkeleton", "Find Skeleton"),
-		LOCTEXT("AnimSequenceBase_FindSkeletonTooltip", "Finds the skeleton for the selected assets in the content browser."),
-		FSlateIcon(FAppStyle::GetAppStyleSetName(), "Icons.Find"),
-		FUIAction(
-			FExecuteAction::CreateSP( this, &FAssetTypeActions_AnimationAsset::ExecuteFindSkeleton, AnimAssets ),
-			FCanExecuteAction()
-			)
-		);
-
-	Section.AddMenuEntry(
-		"AnimSequenceBase_ReplaceSkeleton",
-		LOCTEXT("AnimSequenceBase_ReplaceSkeleton", "Replace Skeleton"),
-		LOCTEXT("AnimSequenceBase_ReplaceSkeletonTooltip", "Associate a different skeleton with the selected animation assets in the content browser."),
-		FSlateIcon(FAppStyle::GetAppStyleSetName(), "Persona.RetargetManager"),
-		FUIAction(
-			FExecuteAction::CreateSP( this, &FAssetTypeActions_AnimationAsset::ExecuteReplaceSkeleton, AnimAssets ),
-			FCanExecuteAction()
-			)
-		);
-}
+PRAGMA_DISABLE_DEPRECATION_WARNINGS
 
 UThumbnailInfo* FAssetTypeActions_AnimationAsset::GetThumbnailInfo(UObject* Asset) const
 {
@@ -249,4 +213,6 @@ bool FAssetTypeActions_AnimationAsset::ReplaceMissingSkeleton(TArray<TObjectPtr<
 	return bWasSkeletonReplaced;
 }
 
+PRAGMA_ENABLE_DEPRECATION_WARNINGS
+
 #undef LOCTEXT_NAMESPACE
diff --git a/Engine/Source/Developer/AssetTools/Private/AssetTypeActions/AssetTypeActions_BlendSpace.h b/Engine/Source/Developer/AssetTools/Private/AssetTypeActions/AssetTypeActions_BlendSpace.h
deleted file mode 100644
diff --git a/Engine/Source/Developer/AssetTools/Private/AssetTypeActions/AssetTypeActions_BlendSpace1D.h b/Engine/Source/Developer/AssetTools/Private/AssetTypeActions/AssetTypeActions_BlendSpace1D.h
deleted file mode 100644
diff --git a/Engine/Source/Developer/AssetTools/Public/AssetTypeActions/AssetTypeActions_AnimationAsset.h b/Engine/Source/Developer/AssetTools/Public/AssetTypeActions/AssetTypeActions_AnimationAsset.h
--- a/Engine/Source/Developer/AssetTools/Public/AssetTypeActions/AssetTypeActions_AnimationAsset.h
+++ b/Engine/Source/Developer/AssetTools/Public/AssetTypeActions/AssetTypeActions_AnimationAsset.h
@@ -9,15 +9,15 @@
 #include "EditorAnimUtils.h"
 #include "SSkeletonWidget.h"
 
-
-class ASSETTOOLS_API FAssetTypeActions_AnimationAsset : public FAssetTypeActions_Base
+class
+// UE_DEPRECATED(5.2, "The AssetDefinition system is replacing AssetTypeActions and UAssetDefinition_AnimationAsset replaced this.  Please see the Conversion Guide in AssetDefinition.h")
+ASSETTOOLS_API FAssetTypeActions_AnimationAsset : public FAssetTypeActions_Base
 {
 public:
 	// IAssetTypeActions Implementation
 	virtual FText GetName() const override { return NSLOCTEXT("AssetTypeActions", "AssetTypeActions_AnimationAsset", "AnimationAsset"); }
 	virtual FColor GetTypeColor() const override { return FColor(80,123,72); }
 	virtual UClass* GetSupportedClass() const override { return UAnimationAsset::StaticClass(); }
-	virtual void GetActions(const TArray<UObject*>& InObjects, struct FToolMenuSection& Section) override;
 	virtual void OpenAssetEditor( const TArray<UObject*>& InObjects, TSharedPtr<class IToolkitHost> EditWithinLevelEditor = TSharedPtr<IToolkitHost>() ) override;
 	virtual bool CanFilter() override { return false; }
 	virtual uint32 GetCategories() override { return EAssetTypeCategories::Animation; }
diff --git a/Engine/Source/Developer/AssetTools/Public/AssetTypeActions/AssetTypeActions_DataAsset.h b/Engine/Source/Developer/AssetTools/Public/AssetTypeActions/AssetTypeActions_DataAsset.h
--- a/Engine/Source/Developer/AssetTools/Public/AssetTypeActions/AssetTypeActions_DataAsset.h
+++ b/Engine/Source/Developer/AssetTools/Public/AssetTypeActions/AssetTypeActions_DataAsset.h
@@ -7,7 +7,9 @@
 #include "Engine/DataAsset.h"
 #include "UObject/WeakObjectPtrTemplates.h"
 
-class ASSETTOOLS_API FAssetTypeActions_DataAsset : public FAssetTypeActions_Base
+class
+// UE_DEPRECATED(5.2, "The AssetDefinition system is replacing AssetTypeActions and UAssetDefinition_DataAsset replaced this.  Please see the Conversion Guide in AssetDefinition.h")
+ASSETTOOLS_API FAssetTypeActions_DataAsset : public FAssetTypeActions_Base
 {
 public:
 	// IAssetTypeActions Implementation
diff --git a/Engine/Source/Developer/AssetTools/Public/AssetTypeActions_CSVAssetBase.h b/Engine/Source/Developer/AssetTools/Public/AssetTypeActions_CSVAssetBase.h
--- a/Engine/Source/Developer/AssetTools/Public/AssetTypeActions_CSVAssetBase.h
+++ b/Engine/Source/Developer/AssetTools/Public/AssetTypeActions_CSVAssetBase.h
@@ -5,7 +5,9 @@
 #include "CoreMinimal.h"
 #include "AssetTypeActions_Base.h"
 
-class FAssetTypeActions_CSVAssetBase : public FAssetTypeActions_Base
+class 
+// UE_DEPRECATED(5.2, "The AssetDefinition system is replacing AssetTypeActions and nothing replaced this, just subclass from UAssetDefinitionDefault.  If you needed the ExecuteFindSourceFileInExplorer, you can now find that in FindSourceFileInExplorer.h.  Please see the Conversion Guide in AssetDefinition.h")
+ASSETTOOLS_API FAssetTypeActions_CSVAssetBase : public FAssetTypeActions_Base
 {
 public:
 	virtual FColor GetTypeColor() const override { return FColor(62, 140, 35); }
diff --git a/Engine/Source/Editor/AssetDefinition/Public/AssetDefinition.h b/Engine/Source/Editor/AssetDefinition/Public/AssetDefinition.h
--- a/Engine/Source/Editor/AssetDefinition/Public/AssetDefinition.h
+++ b/Engine/Source/Editor/AssetDefinition/Public/AssetDefinition.h
@@ -284,6 +284,10 @@ enum class EIncludeClassInFilter : uint8
 
 /**
  * Asset Definitions represent top level assets that are known to the editor.
+ *
+ * -- Conversion Guide --
+ *
+ * 
  */
 UCLASS(Abstract)
 class ASSETDEFINITION_API UAssetDefinition : public UObject
diff --git a/Engine/Source/Editor/AudioEditor/Public/AssetTypeActions/AssetTypeActions_SoundBase.h b/Engine/Source/Editor/AudioEditor/Public/AssetTypeActions/AssetTypeActions_SoundBase.h
--- a/Engine/Source/Editor/AudioEditor/Public/AssetTypeActions/AssetTypeActions_SoundBase.h
+++ b/Engine/Source/Editor/AudioEditor/Public/AssetTypeActions/AssetTypeActions_SoundBase.h
@@ -21,7 +21,7 @@ class USoundBase;
 struct FAssetData;
 
 class
-// UE_DEPRECATED(5.2, "The AssetDefinition system is replacing AssetTypeActions and UAssetDefinition_SoundBase replaced FAssetTypeActions_SoundBase.  Please convert your FAssetTypeActions_SoundBase subclass to subclass from UAssetDefinition_SoundBase instead.")
+// UE_DEPRECATED(5.2, "The AssetDefinition system is replacing AssetTypeActions and UAssetDefinition_SoundBase replaced this.  Please see the Conversion Guide in AssetDefinition.h")
 AUDIOEDITOR_API
 FAssetTypeActions_SoundBase : public FAssetTypeActions_Base
 {
diff --git a/Engine/Source/Editor/UMGEditor/Public/IHasDesignerExtensibility.h b/Engine/Source/Editor/UMGEditor/Public/IHasDesignerExtensibility.h
--- a/Engine/Source/Editor/UMGEditor/Public/IHasDesignerExtensibility.h
+++ b/Engine/Source/Editor/UMGEditor/Public/IHasDesignerExtensibility.h
@@ -11,9 +11,20 @@
 class UMGEDITOR_API IDesignerExtensionFactory
 {
 public:
+	virtual ~IDesignerExtensionFactory() { }
 	virtual TSharedRef<FDesignerExtension> CreateDesignerExtension() const = 0;
 };
 
+class FDesignerExtensibilityManager;
+
+/** Indicates that a class has a designer that is extensible */
+class IHasDesignerExtensibility
+{
+public:
+	virtual ~IHasDesignerExtensibility() { }
+	virtual TSharedPtr<FDesignerExtensibilityManager> GetDesignerExtensibilityManager() = 0;
+};
+
 /**
  * Designer Extensibility Manager keep a series of Designer Extensions. See FDesignerExtension class for more information.
  */
@@ -57,11 +68,3 @@ private:
 	TArray<TSharedRef<FDesignerExtension>> ExternalExtensions;
 	TArray<TSharedRef<IDesignerExtensionFactory>> ExternalExtensionFactories;
 };
-
-/** Indicates that a class has a designer that is extensible */
-class IHasDesignerExtensibility
-{
-public:
-	virtual TSharedPtr<FDesignerExtensibilityManager> GetDesignerExtensibilityManager() = 0;
-};
-
diff --git a/Engine/Source/Editor/UnrealEd/Public/SSkeletonWidget.h b/Engine/Source/Editor/UnrealEd/Public/SSkeletonWidget.h
--- a/Engine/Source/Editor/UnrealEd/Public/SSkeletonWidget.h
+++ b/Engine/Source/Editor/UnrealEd/Public/SSkeletonWidget.h
@@ -439,7 +439,12 @@ public:
 	SLATE_END_ARGS()
 
 	void UNREALED_API Construct(const FArguments& InArgs);
-	SReplaceMissingSkeletonDialog() : UserResponse(EAppReturnType::Cancel){}
+	
+	SReplaceMissingSkeletonDialog()
+		: UserResponse(EAppReturnType::Cancel)
+		, bWasSkeletonReplaced(false)
+	{
+	}
 
 	bool UNREALED_API ShowModal();
 
diff --git a/Engine/Source/Runtime/Engine/Classes/Animation/AnimationAsset.h b/Engine/Source/Runtime/Engine/Classes/Animation/AnimationAsset.h
--- a/Engine/Source/Runtime/Engine/Classes/Animation/AnimationAsset.h
+++ b/Engine/Source/Runtime/Engine/Classes/Animation/AnimationAsset.h
@@ -1116,7 +1116,7 @@ public:
 	 * Parent Asset related function. Used by editor
 	 */
 	void SetParentAsset(UAnimationAsset* InParentAsset);
-	bool HasParentAsset() { return ParentAsset != nullptr;  }
+	bool HasParentAsset() const { return ParentAsset != nullptr; }
 	bool RemapAsset(UAnimationAsset* SourceAsset, UAnimationAsset* TargetAsset);
 	// we have to update whenever we have anything loaded
 	void UpdateParentAsset();
@@ -1138,6 +1138,7 @@ public:
 	//~ Begin UObject Interface.
 #if WITH_EDITOR
 	virtual void PostEditChangeProperty(FPropertyChangedEvent& PropertyChangedEvent) override;
+	virtual void GetAssetRegistryTags(TArray<FAssetRegistryTag>& OutTags) const override;
 	virtual EDataValidationResult IsDataValid(TArray<FText>& ValidationErrors) override;
 #endif // WITH_EDITOR
 
diff --git a/Engine/Source/Runtime/Engine/Classes/Interfaces/Interface_AssetUserData.h b/Engine/Source/Runtime/Engine/Classes/Interfaces/Interface_AssetUserData.h
--- a/Engine/Source/Runtime/Engine/Classes/Interfaces/Interface_AssetUserData.h
+++ b/Engine/Source/Runtime/Engine/Classes/Interfaces/Interface_AssetUserData.h
@@ -25,12 +25,12 @@ class IInterface_AssetUserData
 
 	virtual UAssetUserData* GetAssetUserDataOfClass(TSubclassOf<UAssetUserData> InUserDataClass)
 	{
-		return NULL;
+		return nullptr;
 	}
 
 	virtual const TArray<UAssetUserData*>* GetAssetUserDataArray() const
 	{
-		return NULL;
+		return nullptr;
 	}
 
 	template<typename T>
diff --git a/Engine/Source/Runtime/Engine/Private/AnimationAsset.cpp b/Engine/Source/Runtime/Engine/Private/AnimationAsset.cpp
--- a/Engine/Source/Runtime/Engine/Private/AnimationAsset.cpp
+++ b/Engine/Source/Runtime/Engine/Private/AnimationAsset.cpp
@@ -722,6 +722,18 @@ void UAnimationAsset::PostEditChangeProperty(FPropertyChangedEvent& PropertyChan
 	Super::PostEditChangeProperty(PropertyChangedEvent);
 }
 
+void UAnimationAsset::GetAssetRegistryTags(TArray<FAssetRegistryTag>& OutTags) const
+{
+	Super::GetAssetRegistryTags(OutTags);
+	
+	for (const UAssetUserData* UserData : AssetUserData)
+	{
+		UserData->GetAssetRegistryTags(OutTags);
+	}
+	
+	OutTags.Add( FAssetRegistryTag("HasParentAsset", HasParentAsset() ? TEXT("True") : TEXT("False"), FAssetRegistryTag::TT_Hidden) );
+}
+
 EDataValidationResult UAnimationAsset::IsDataValid(TArray<FText>& ValidationErrors)
 {
 	EDataValidationResult Result = UObject::IsDataValid(ValidationErrors);
diff --git a/Engine/Source/Runtime/LevelSequence/Public/AnimSequenceLevelSequenceLink.h b/Engine/Source/Runtime/LevelSequence/Public/AnimSequenceLevelSequenceLink.h
--- a/Engine/Source/Runtime/LevelSequence/Public/AnimSequenceLevelSequenceLink.h
+++ b/Engine/Source/Runtime/LevelSequence/Public/AnimSequenceLevelSequenceLink.h
@@ -16,7 +16,7 @@ public:
 	UPROPERTY(BlueprintReadWrite, Category = Property)
 	FGuid SkelTrackGuid;
 
-	UPROPERTY(BlueprintReadWrite, Category = Property)
+	UPROPERTY(BlueprintReadWrite, AssetRegistrySearchable, Category = Property)
 	FSoftObjectPath PathToLevelSequence;
 
 	void SetLevelSequence(ULevelSequence* InLevelSequence);
