Fixing build issue with deprecated function.

diff --git a/Engine/Plugins/Editor/EngineAssetDefinitions/Source/Private/Animation/AssetDefinition_PoseAsset.cpp b/Engine/Plugins/Editor/EngineAssetDefinitions/Source/Private/Animation/AssetDefinition_PoseAsset.cpp
--- a/Engine/Plugins/Editor/EngineAssetDefinitions/Source/Private/Animation/AssetDefinition_PoseAsset.cpp
+++ b/Engine/Plugins/Editor/EngineAssetDefinitions/Source/Private/Animation/AssetDefinition_PoseAsset.cpp
@@ -25,12 +25,15 @@ namespace MenuExtension_PoseAsset
 		{
 			if (UAnimSequence* SourceAnimation = PoseAsset->SourceAnimation)
 			{
-				if (!PoseAsset->SourceAnimationRawDataGUID.IsValid() || PoseAsset->SourceAnimationRawDataGUID != SourceAnimation->GenerateGuidFromModel())
+				if (IAnimationDataModel* SourceAnimationDataModel = SourceAnimation->GetDataModel())
 				{
-					if (PoseAsset->GetSkeleton()->IsCompatibleForEditor(SourceAnimation->GetSkeleton()))
+					if (!PoseAsset->SourceAnimationRawDataGUID.IsValid() || PoseAsset->SourceAnimationRawDataGUID != SourceAnimationDataModel->GenerateGuid())
 					{
-						PoseAsset->Modify();
-						PoseAsset->UpdatePoseFromAnimation(SourceAnimation);
+						if (PoseAsset->GetSkeleton()->IsCompatibleForEditor(SourceAnimation->GetSkeleton()))
+						{
+							PoseAsset->Modify();
+							PoseAsset->UpdatePoseFromAnimation(SourceAnimation);
+						}
 					}
 				}
 			}
