Editor Performance - Fixing a long standing issue with asset thumbnails, they've been loading the assets for a pretty long time now, because of the thumbnail editing mode.  Note that thumbnail editing mode was not needed to be enabled and yet the asset would still load.  This was a cause for a lot of hitches in the Content Browser.

AssetDefinition - Working on the asset definitions for blueprints/class types, iterating on the Merge API.  Bringing over more of the blueprint features into its asset definition, unregistered the Blueprint AssetType Actions.

Asset Filters - We now support categories per filter, this enables a single asset definition to correctly specify multiple filters.  For example, blueprints actually supply 4 filters, one for blueprints, but also blueprint interfaces, macros and function libraries.

diff --git a/Engine/Plugins/Editor/EngineAssetDefinitions/Source/EngineAssetDefinitions.Build.cs b/Engine/Plugins/Editor/EngineAssetDefinitions/Source/EngineAssetDefinitions.Build.cs
--- a/Engine/Plugins/Editor/EngineAssetDefinitions/Source/EngineAssetDefinitions.Build.cs
+++ b/Engine/Plugins/Editor/EngineAssetDefinitions/Source/EngineAssetDefinitions.Build.cs
@@ -33,6 +33,8 @@ public class EngineAssetDefinitions : ModuleRules
 				"Foliage",
 				"Landscape",
 				"PhysicsCore",
+				"Kismet",
+				"GameProjectGeneration"
 			}
 		);
 		
diff --git a/Engine/Plugins/Editor/EngineAssetDefinitions/Source/Private/Animation/AssetDefinition_AnimationAsset.cpp b/Engine/Plugins/Editor/EngineAssetDefinitions/Source/Private/Animation/AssetDefinition_AnimationAsset.cpp
--- a/Engine/Plugins/Editor/EngineAssetDefinitions/Source/Private/Animation/AssetDefinition_AnimationAsset.cpp
+++ b/Engine/Plugins/Editor/EngineAssetDefinitions/Source/Private/Animation/AssetDefinition_AnimationAsset.cpp
@@ -148,15 +148,7 @@ EAssetCommandResult UAssetDefinition_AnimationAsset::OpenAssets(const FAssetOpen
 
 UThumbnailInfo* UAssetDefinition_AnimationAsset::LoadThumbnailInfo(const FAssetData& InAssetData) const
 {
-	UAnimationAsset* Anim = CastChecked<UAnimationAsset>(InAssetData.GetAsset());
-	UThumbnailInfo* ThumbnailInfo = Anim->ThumbnailInfo;
-	if (ThumbnailInfo == nullptr)
-	{
-		ThumbnailInfo = NewObject<USceneThumbnailInfo>(Anim, NAME_None, RF_Transactional);
-		Anim->ThumbnailInfo = ThumbnailInfo;
-	}
-
-	return ThumbnailInfo;
+	return UE::Editor::FindOrCreateThumbnailInfo(InAssetData.GetAsset(), USceneThumbnailInfo::StaticClass());
 }
 
 // Menu Extensions
diff --git a/Engine/Plugins/Editor/EngineAssetDefinitions/Source/Private/Animation/AssetDefinition_SkeletalMesh.cpp b/Engine/Plugins/Editor/EngineAssetDefinitions/Source/Private/Animation/AssetDefinition_SkeletalMesh.cpp
--- a/Engine/Plugins/Editor/EngineAssetDefinitions/Source/Private/Animation/AssetDefinition_SkeletalMesh.cpp
+++ b/Engine/Plugins/Editor/EngineAssetDefinitions/Source/Private/Animation/AssetDefinition_SkeletalMesh.cpp
@@ -77,15 +77,7 @@ EAssetCommandResult UAssetDefinition_SkeletalMesh::OpenAssets(const FAssetOpenAr
 
 UThumbnailInfo* UAssetDefinition_SkeletalMesh::LoadThumbnailInfo(const FAssetData& InAssetData) const
 {
-	USkeletalMesh* SkeletalMesh = CastChecked<USkeletalMesh>(InAssetData.GetAsset());
-	UThumbnailInfo* ThumbnailInfo = SkeletalMesh->GetThumbnailInfo();
-	if ( ThumbnailInfo == nullptr )
-	{
-		ThumbnailInfo = NewObject<USceneThumbnailInfo>(SkeletalMesh, NAME_None, RF_Transactional);
-		SkeletalMesh->SetThumbnailInfo(ThumbnailInfo);
-	}
-
-	return ThumbnailInfo;
+	return UE::Editor::FindOrCreateThumbnailInfo(InAssetData.GetAsset(), USceneThumbnailInfo::StaticClass());
 }
 
 EVisibility UAssetDefinition_SkeletalMesh::GetThumbnailSkinningOverlayVisibility(const FAssetData AssetData) const
diff --git a/Engine/Plugins/Editor/EngineAssetDefinitions/Source/Private/AssetDefinition_Blueprint.cpp b/Engine/Plugins/Editor/EngineAssetDefinitions/Source/Private/AssetDefinition_Blueprint.cpp
deleted file mode 100644
diff --git a/Engine/Plugins/Editor/EngineAssetDefinitions/Source/Private/AssetDefinition_PhysicsAsset.cpp b/Engine/Plugins/Editor/EngineAssetDefinitions/Source/Private/AssetDefinition_PhysicsAsset.cpp
--- a/Engine/Plugins/Editor/EngineAssetDefinitions/Source/Private/AssetDefinition_PhysicsAsset.cpp
+++ b/Engine/Plugins/Editor/EngineAssetDefinitions/Source/Private/AssetDefinition_PhysicsAsset.cpp
@@ -9,15 +9,7 @@
 
 UThumbnailInfo* UAssetDefinition_PhysicsAsset::LoadThumbnailInfo(const FAssetData& InAsset) const
 {
-	UPhysicsAsset* PhysicsAsset = CastChecked<UPhysicsAsset>(InAsset.GetAsset());
-	UThumbnailInfo* ThumbnailInfo = PhysicsAsset->ThumbnailInfo;
-	if (ThumbnailInfo == nullptr)
-	{
-		ThumbnailInfo = NewObject<USceneThumbnailInfo>(PhysicsAsset, NAME_None, RF_Transactional);
-		PhysicsAsset->ThumbnailInfo = ThumbnailInfo;
-	}
-
-	return ThumbnailInfo;
+	return UE::Editor::FindOrCreateThumbnailInfo(InAsset.GetAsset(), USceneThumbnailInfo::StaticClass());
 }
 
 EAssetCommandResult UAssetDefinition_PhysicsAsset::OpenAssets(const FAssetOpenArgs& OpenArgs) const
diff --git a/Engine/Plugins/Editor/EngineAssetDefinitions/Source/Private/AssetDefinition_StaticMesh.cpp b/Engine/Plugins/Editor/EngineAssetDefinitions/Source/Private/AssetDefinition_StaticMesh.cpp
--- a/Engine/Plugins/Editor/EngineAssetDefinitions/Source/Private/AssetDefinition_StaticMesh.cpp
+++ b/Engine/Plugins/Editor/EngineAssetDefinitions/Source/Private/AssetDefinition_StaticMesh.cpp
@@ -19,20 +19,7 @@
 
 UThumbnailInfo* UAssetDefinition_StaticMesh::LoadThumbnailInfo(const FAssetData& InAssetData) const
 {
-	UStaticMesh* StaticMesh = Cast<UStaticMesh>(InAssetData.GetAsset());
-	if (ensure(StaticMesh))
-	{
-		UThumbnailInfo* ThumbnailInfo = StaticMesh->ThumbnailInfo;
-		if ( ThumbnailInfo == nullptr )
-		{
-			ThumbnailInfo = NewObject<USceneThumbnailInfo>(StaticMesh, NAME_None, RF_Transactional);
-			StaticMesh->ThumbnailInfo = ThumbnailInfo;
-		}
-
-		return ThumbnailInfo;
-	}
-
-	return nullptr;
+	return UE::Editor::FindOrCreateThumbnailInfo(InAssetData.GetAsset(), USceneThumbnailInfo::StaticClass());
 }
 
 EAssetCommandResult UAssetDefinition_StaticMesh::OpenAssets(const FAssetOpenArgs& OpenArgs) const
diff --git a/Engine/Plugins/Editor/EngineAssetDefinitions/Source/Private/AssetDefinition_World.cpp b/Engine/Plugins/Editor/EngineAssetDefinitions/Source/Private/AssetDefinition_World.cpp
--- a/Engine/Plugins/Editor/EngineAssetDefinitions/Source/Private/AssetDefinition_World.cpp
+++ b/Engine/Plugins/Editor/EngineAssetDefinitions/Source/Private/AssetDefinition_World.cpp
@@ -71,20 +71,7 @@ EAssetCommandResult UAssetDefinition_World::OpenAssets(const FAssetOpenArgs& Ope
 
 UThumbnailInfo* UAssetDefinition_World::LoadThumbnailInfo(const FAssetData& InAsset) const
 {
-	if (UObject* WorldInstance = InAsset.GetAsset())
-	{
-		UWorld* World = CastChecked<UWorld>(WorldInstance);
-		UThumbnailInfo* ThumbnailInfo = World->ThumbnailInfo;
-		if (ThumbnailInfo == nullptr)
-		{
-			ThumbnailInfo = NewObject<UWorldThumbnailInfo>(World, NAME_None, RF_Transactional);
-			World->ThumbnailInfo = ThumbnailInfo;
-		}
-		
-		return ThumbnailInfo;
-	}
-
-	return nullptr;
+	return UE::Editor::FindOrCreateThumbnailInfo(InAsset.GetAsset(), UWorldThumbnailInfo::StaticClass());
 }
 
 FAssetSupportResponse UAssetDefinition_World::CanRename(const FAssetData& InAsset) const
diff --git a/Engine/Plugins/Editor/EngineAssetDefinitions/Source/Private/Material/AssetDefinition_MaterialFunction.cpp b/Engine/Plugins/Editor/EngineAssetDefinitions/Source/Private/Material/AssetDefinition_MaterialFunction.cpp
--- a/Engine/Plugins/Editor/EngineAssetDefinitions/Source/Private/Material/AssetDefinition_MaterialFunction.cpp
+++ b/Engine/Plugins/Editor/EngineAssetDefinitions/Source/Private/Material/AssetDefinition_MaterialFunction.cpp
@@ -34,15 +34,7 @@ EAssetCommandResult UAssetDefinition_MaterialFunction::OpenAssets(const FAssetOp
 
 UThumbnailInfo* UAssetDefinition_MaterialFunction::LoadThumbnailInfo(const FAssetData& InAsset) const
 {
-	UMaterialFunctionInterface* MaterialFunc = CastChecked<UMaterialFunctionInterface>(InAsset.GetAsset());
-	UThumbnailInfo* ThumbnailInfo = MaterialFunc->ThumbnailInfo;
-	if (ThumbnailInfo == nullptr)
-	{
-		ThumbnailInfo = NewObject<USceneThumbnailInfoWithPrimitive>(MaterialFunc, NAME_None, RF_Transactional);
-		MaterialFunc->ThumbnailInfo = ThumbnailInfo;
-	}
-
-	return ThumbnailInfo;
+	return UE::Editor::FindOrCreateThumbnailInfo(InAsset.GetAsset(), USceneThumbnailInfoWithPrimitive::StaticClass());
 }
 
 // Menu Extensions
diff --git a/Engine/Plugins/Editor/EngineAssetDefinitions/Source/Private/Material/AssetDefinition_MaterialFunctionInstance.cpp b/Engine/Plugins/Editor/EngineAssetDefinitions/Source/Private/Material/AssetDefinition_MaterialFunctionInstance.cpp
--- a/Engine/Plugins/Editor/EngineAssetDefinitions/Source/Private/Material/AssetDefinition_MaterialFunctionInstance.cpp
+++ b/Engine/Plugins/Editor/EngineAssetDefinitions/Source/Private/Material/AssetDefinition_MaterialFunctionInstance.cpp
@@ -25,15 +25,7 @@ EAssetCommandResult UAssetDefinition_MaterialFunctionInstance::OpenAssets(const
 
 UThumbnailInfo* UAssetDefinition_MaterialFunctionInstance::LoadThumbnailInfo(const FAssetData& InAsset) const
 {
-	UMaterialFunctionInstance* MaterialFunc = CastChecked<UMaterialFunctionInstance>(InAsset.GetAsset());
-	UThumbnailInfo* ThumbnailInfo = MaterialFunc->ThumbnailInfo;
-	if (ThumbnailInfo == nullptr)
-	{
-		ThumbnailInfo = NewObject<USceneThumbnailInfoWithPrimitive>(MaterialFunc, NAME_None, RF_Transactional);
-		MaterialFunc->ThumbnailInfo = ThumbnailInfo;
-	}
-
-	return ThumbnailInfo;
+	return UE::Editor::FindOrCreateThumbnailInfo(InAsset.GetAsset(), USceneThumbnailInfoWithPrimitive::StaticClass());
 }
 
 // Menu Extensions
diff --git a/Engine/Plugins/Editor/EngineAssetDefinitions/Source/Private/Material/AssetDefinition_MaterialInterface.cpp b/Engine/Plugins/Editor/EngineAssetDefinitions/Source/Private/Material/AssetDefinition_MaterialInterface.cpp
--- a/Engine/Plugins/Editor/EngineAssetDefinitions/Source/Private/Material/AssetDefinition_MaterialInterface.cpp
+++ b/Engine/Plugins/Editor/EngineAssetDefinitions/Source/Private/Material/AssetDefinition_MaterialInterface.cpp
@@ -17,21 +17,21 @@
 
 UThumbnailInfo* UAssetDefinition_MaterialInterface::LoadThumbnailInfo(const FAssetData& InAsset) const
 {
-	UMaterialInterface* MaterialInterface = CastChecked<UMaterialInterface>(InAsset.GetAsset());
-	USceneThumbnailInfoWithPrimitive* ThumbnailInfo = Cast<USceneThumbnailInfoWithPrimitive>(MaterialInterface->ThumbnailInfo);
-	if ( ThumbnailInfo == nullptr )
+	if (UMaterialInterface* MaterialInterface = Cast<UMaterialInterface>(InAsset.GetAsset()))
 	{
-		ThumbnailInfo = NewObject<USceneThumbnailInfoWithPrimitive>(MaterialInterface, NAME_None, RF_Transactional);
-		MaterialInterface->ThumbnailInfo = ThumbnailInfo;
+		if (USceneThumbnailInfoWithPrimitive* ThumbnailInfo = UE::Editor::FindOrCreateThumbnailInfo<USceneThumbnailInfoWithPrimitive>(MaterialInterface))
+		{
+			const UMaterial* Material = MaterialInterface->GetBaseMaterial();
+			if (Material && Material->bUsedWithParticleSprites)
+			{
+				ThumbnailInfo->DefaultPrimitiveType = TPT_Plane;
+			}
+
+			return ThumbnailInfo;
+		}
 	}
-	
-	const UMaterial* Material = MaterialInterface->GetBaseMaterial();
-	if (Material && Material->bUsedWithParticleSprites)
-    {
-    	ThumbnailInfo->DefaultPrimitiveType = TPT_Plane;
-    }
 
-	return ThumbnailInfo;
+	return nullptr;
 }
 
 // Menu Extensions
diff --git a/Engine/Plugins/Editor/EngineAssetDefinitions/Source/Private/Script/AssetDefinition_Blueprint.cpp b/Engine/Plugins/Editor/EngineAssetDefinitions/Source/Private/Script/AssetDefinition_Blueprint.cpp
new file mode 100644
--- /dev/null
+++ b/Engine/Plugins/Editor/EngineAssetDefinitions/Source/Private/Script/AssetDefinition_Blueprint.cpp
@@ -0,0 +1,372 @@
+// Copyright Epic Games, Inc. All Rights Reserved.
+
+#include "Script/AssetDefinition_Blueprint.h"
+
+#include "AssetDefinitionRegistry.h"
+#include "AssetToolsModule.h"
+#include "Blueprint/BlueprintSupport.h"
+#include "BlueprintEditorModule.h"
+#include "ContentBrowserMenuContexts.h"
+#include "ContentBrowserModule.h"
+#include "IAssetTools.h"
+#include "ToolMenus.h"
+#include "Factories/BlueprintFactory.h"
+#include "Kismet2/KismetEditorUtilities.h"
+#include "IContentBrowserSingleton.h"
+#include "Algo/AllOf.h"
+#include "Logging/MessageLog.h"
+#include "Misc/MessageDialog.h"
+#include "Misc/PackageName.h"
+#include "ToolMenu.h"
+#include "ToolMenuSection.h"
+#include "Toolkits/AssetEditorToolkit.h"
+#include "BlueprintEditor.h"
+#include "SBlueprintDiff.h"
+#include "Kismet2/BlueprintEditorUtils.h"
+#include "ThumbnailRendering/SceneThumbnailInfo.h"
+#include "Misc/AssetFilterData.h"
+
+#define LOCTEXT_NAMESPACE "UAssetDefinition_Blueprint"
+
+FText UAssetDefinition_Blueprint::GetAssetDisplayName(const FAssetData& AssetData) const
+{
+	FString OutBlueprintType;
+	if (AssetData.GetTagValue(GET_MEMBER_NAME_CHECKED(UBlueprint, BlueprintType), OutBlueprintType))
+	{
+		if (OutBlueprintType == TEXT("BPTYPE_Interface"))
+		{
+			return LOCTEXT("AssetTypeActions_BlueprintInterface", "Blueprint Interface");
+		}
+		else if (OutBlueprintType == TEXT("BPTYPE_MacroLibrary"))
+        {
+			return LOCTEXT("AssetTypeActions_BlueprintMacroLibrary", "Blueprint Macro Library");
+        }
+		else if (OutBlueprintType == TEXT("BPTYPE_FunctionLibrary"))
+		{
+			return LOCTEXT("AssetTypeActions_BlueprintFunctionLibrary", "Blueprint Function Library");
+		}
+	}
+
+	return Super::GetAssetDisplayName(AssetData);
+}
+
+FText UAssetDefinition_Blueprint::GetAssetDescription(const FAssetData& AssetData) const
+{
+	FString Description = AssetData.GetTagValueRef<FString>(GET_MEMBER_NAME_CHECKED(UBlueprint, BlueprintDescription));
+	if (!Description.IsEmpty())
+	{
+		Description.ReplaceInline(TEXT("\\n"), TEXT("\n"));
+		return FText::FromString(MoveTemp(Description));
+	}
+
+	return FText::GetEmpty();
+}
+
+EAssetCommandResult UAssetDefinition_Blueprint::Merge(const FAssetAutomaticMergeArgs& MergeArgs) const
+{
+	UBlueprint* AsBlueprint = CastChecked<UBlueprint>(MergeArgs.LocalAsset);
+	
+	// Kludge to get the merge panel in the blueprint editor to show up:
+	bool Success = GEditor->GetEditorSubsystem<UAssetEditorSubsystem>()->OpenEditorForAsset(AsBlueprint);
+	if (Success)
+	{
+		FBlueprintEditorModule& BlueprintEditorModule = FModuleManager::LoadModuleChecked<FBlueprintEditorModule>( "Kismet" );
+
+		FBlueprintEditor* BlueprintEditor = static_cast<FBlueprintEditor*>(GEditor->GetEditorSubsystem<UAssetEditorSubsystem>()->FindEditorForAsset(AsBlueprint, false));
+		BlueprintEditor->CreateMergeToolTab();
+	}
+
+	return EAssetCommandResult::Handled;
+}
+
+EAssetCommandResult UAssetDefinition_Blueprint::Merge(const FAssetManualMergeArgs& MergeArgs) const
+{
+	UBlueprint* AsBlueprint = CastChecked<UBlueprint>(MergeArgs.LocalAsset);
+	check(MergeArgs.LocalAsset->GetClass() == MergeArgs.BaseAsset->GetClass());
+	check(MergeArgs.LocalAsset->GetClass() == MergeArgs.RemoteAsset->GetClass());
+
+	if (GEditor->GetEditorSubsystem<UAssetEditorSubsystem>()->OpenEditorForAsset(AsBlueprint))
+	{
+		FBlueprintEditorModule& BlueprintEditorModule = FModuleManager::LoadModuleChecked<FBlueprintEditorModule>("Kismet");
+
+		FBlueprintEditor* BlueprintEditor = static_cast<FBlueprintEditor*>(GEditor->GetEditorSubsystem<UAssetEditorSubsystem>()->FindEditorForAsset(AsBlueprint, /*bFocusIfOpen =*/false));
+		BlueprintEditor->CreateMergeToolTab(Cast<UBlueprint>(MergeArgs.BaseAsset), Cast<UBlueprint>(MergeArgs.RemoteAsset), 
+			FOnMergeResolved::CreateLambda([ResolutionCallback = MergeArgs.ResolutionCallback](UPackage* MergedPackage, EMergeResult::Type Result)
+			{
+				FAssetMergeResults Results;
+				Results.MergedPackage = MergedPackage;
+				Results.Result = static_cast<EAssetMergeResult>(Result);
+				ResolutionCallback.Execute(Results);
+			})
+		);		
+	}
+	
+	return EAssetCommandResult::Handled;
+}
+
+EAssetCommandResult UAssetDefinition_Blueprint::PerformAssetDiff(const FAssetDiffArgs& DiffArgs) const
+{
+	UBlueprint* OldBlueprint = CastChecked<UBlueprint>(DiffArgs.OldAsset);
+	UBlueprint* NewBlueprint = CastChecked<UBlueprint>(DiffArgs.NewAsset);
+
+	// sometimes we're comparing different revisions of one single asset (other 
+	// times we're comparing two completely separate assets altogether)
+	bool bIsSingleAsset = (NewBlueprint->GetName() == OldBlueprint->GetName());
+
+	FText WindowTitle = LOCTEXT("NamelessBlueprintDiff", "Blueprint Diff");
+	// if we're diffing one asset against itself 
+	if (bIsSingleAsset)
+	{
+		// identify the assumed single asset in the window's title
+		WindowTitle = FText::Format(LOCTEXT("Blueprint Diff", "{0} - Blueprint Diff"), FText::FromString(NewBlueprint->GetName()));
+	}
+
+	SBlueprintDiff::CreateDiffWindow(WindowTitle, OldBlueprint, NewBlueprint, DiffArgs.OldRevision, DiffArgs.NewRevision);
+
+	return EAssetCommandResult::Handled;
+}
+
+UThumbnailInfo* UAssetDefinition_Blueprint::LoadThumbnailInfo(const FAssetData& InAssetData) const
+{
+	return UE::Editor::FindOrCreateThumbnailInfo(InAssetData.GetAsset(), USceneThumbnailInfo::StaticClass());
+}
+
+void UAssetDefinition_Blueprint::BuildFilters(TArray<FAssetFilterData>& OutFilters) const
+{
+	Super::BuildFilters(OutFilters);
+
+	{
+		FAssetFilterData Filter;
+		Filter.Name = GetAssetClass().ToSoftObjectPath().ToString() + TEXT("Interface");
+		Filter.DisplayText = LOCTEXT("AssetTypeActions_BlueprintInterface", "Blueprint Interface");
+		Filter.FilterCategories = { EAssetCategoryPaths::Blueprint };
+		Filter.Filter.ClassPaths.Add(GetAssetClass().ToSoftObjectPath().GetAssetPath());
+		Filter.Filter.TagsAndValues.Add(GET_MEMBER_NAME_CHECKED(UBlueprint, BlueprintType), FString(TEXT("BPTYPE_Interface")));
+		OutFilters.Add(MoveTemp(Filter));
+	}
+
+	{
+		FAssetFilterData Filter;
+		Filter.Name = GetAssetClass().ToSoftObjectPath().ToString() + TEXT("MacroLibrary");
+		Filter.DisplayText = LOCTEXT("AssetTypeActions_BlueprintMacroLibrary", "Blueprint Macro Library");
+		Filter.FilterCategories = { EAssetCategoryPaths::Blueprint };
+		Filter.Filter.ClassPaths.Add(GetAssetClass().ToSoftObjectPath().GetAssetPath());
+		Filter.Filter.TagsAndValues.Add(GET_MEMBER_NAME_CHECKED(UBlueprint, BlueprintType), FString(TEXT("BPTYPE_MacroLibrary")));
+		OutFilters.Add(MoveTemp(Filter));
+	}
+
+	{
+		FAssetFilterData Filter;
+		Filter.Name = GetAssetClass().ToSoftObjectPath().ToString() + TEXT("FunctionLibrary");
+		Filter.DisplayText = LOCTEXT("AssetTypeActions_BlueprintFunctionLibrary", "Blueprint Function Library");
+		Filter.Filter.ClassPaths.Add(GetAssetClass().ToSoftObjectPath().GetAssetPath());
+		Filter.FilterCategories = { EAssetCategoryPaths::Blueprint };
+		Filter.Filter.TagsAndValues.Add(GET_MEMBER_NAME_CHECKED(UBlueprint, BlueprintType), FString(TEXT("BPTYPE_FunctionLibrary")));
+		OutFilters.Add(MoveTemp(Filter));
+	}
+}
+
+EAssetCommandResult UAssetDefinition_Blueprint::OpenAssets(const FAssetOpenArgs& OpenArgs) const
+{
+	TArray<FAssetData> OutAssetsThatFailedToLoad;
+	for (UBlueprint* Blueprint : OpenArgs.LoadObjects<UBlueprint>({}, &OutAssetsThatFailedToLoad))
+	{
+		bool bLetOpen = true;
+		if (!Blueprint->SkeletonGeneratedClass || !Blueprint->GeneratedClass)
+		{
+			bLetOpen = EAppReturnType::Yes == FMessageDialog::Open(EAppMsgType::YesNo, LOCTEXT("FailedToLoadBlueprintWithContinue", "Blueprint could not be loaded because it derives from an invalid class.  Check to make sure the parent class for this blueprint hasn't been removed! Do you want to continue (it can crash the editor)?"));
+		}
+		if (bLetOpen)
+		{
+			FBlueprintEditorModule& BlueprintEditorModule = FModuleManager::LoadModuleChecked<FBlueprintEditorModule>("Kismet");
+			TSharedRef< IBlueprintEditor > NewKismetEditor = BlueprintEditorModule.CreateBlueprintEditor(OpenArgs.GetToolkitMode(), OpenArgs.ToolkitHost, Blueprint, FBlueprintEditorUtils::ShouldOpenWithDataOnlyEditor(Blueprint));
+		}
+	}
+
+	for (const FAssetData& UnableToLoadAsset : OutAssetsThatFailedToLoad)
+	{
+		FMessageDialog::Open( EAppMsgType::Ok,
+			FText::Format(
+				LOCTEXT("FailedToLoadBlueprint", "Blueprint '{0}' could not be loaded because it derives from an invalid class.  Check to make sure the parent class for this blueprint hasn't been removed!"),
+				FText::FromName(UnableToLoadAsset.PackagePath)
+			)
+		);
+	}
+
+	return EAssetCommandResult::Handled;
+}
+
+TWeakPtr<IClassTypeActions> UAssetDefinition_Blueprint::GetClassTypeActions(const FAssetData& AssetData) const
+{
+	// Blueprints get the class type actions for their parent native class - this avoids us having to load the blueprint
+	UClass* ParentClass = nullptr;
+	FString ParentClassName;
+	if (!AssetData.GetTagValue(FBlueprintTags::NativeParentClassPath, ParentClassName))
+	{
+		AssetData.GetTagValue(FBlueprintTags::ParentClassPath, ParentClassName);
+	}
+	
+	if (!ParentClassName.IsEmpty())
+	{
+		UObject* Outer = nullptr;
+		ResolveName(Outer, ParentClassName, false, false);
+		ParentClass = FindObject<UClass>(Outer, *ParentClassName);
+	}
+
+	if (ParentClass)
+	{
+		FAssetToolsModule& AssetToolsModule = FAssetToolsModule::GetModule();
+		return AssetToolsModule.Get().GetClassTypeActionsForClass(ParentClass);
+	}
+
+	return nullptr;
+}
+
+UFactory* UAssetDefinition_Blueprint::GetFactoryForBlueprintType(UBlueprint* InBlueprint) const
+{
+	UBlueprintFactory* BlueprintFactory = NewObject<UBlueprintFactory>();
+	BlueprintFactory->ParentClass = InBlueprint->GeneratedClass;
+	return BlueprintFactory;
+}
+
+// Menu Extensions
+//--------------------------------------------------------------------
+
+namespace MenuExtension_Blueprint
+{
+	static bool CanExecuteNewDerivedBlueprint(const FToolMenuContext& MenuContext, const FAssetData* SelectedBlueprintPtr)
+	{
+		const uint32 BPFlags = SelectedBlueprintPtr->GetTagValueRef<uint32>(FBlueprintTags::ClassFlags);
+		if ((BPFlags & (CLASS_Deprecated)) == 0)
+		{
+			return true;
+		}
+
+		return false;
+	}
+
+	static void ExecuteNewDerivedBlueprint(const FToolMenuContext& MenuContext, const FAssetData* SelectedBlueprintPtr)
+	{
+		if (UBlueprint* ParentBlueprint = Cast<UBlueprint>(SelectedBlueprintPtr->GetAsset()))
+		{
+			UClass* TargetParentClass = ParentBlueprint->GeneratedClass;
+
+			if (!FKismetEditorUtilities::CanCreateBlueprintOfClass(TargetParentClass))
+			{
+				FMessageDialog::Open( EAppMsgType::Ok, LOCTEXT("InvalidClassToMakeBlueprintFrom", "Invalid class with which to make a Blueprint."));
+				return;
+			}
+
+			FString Name;
+			FString PackageName;
+			FAssetToolsModule& AssetToolsModule = FModuleManager::LoadModuleChecked<FAssetToolsModule>("AssetTools");
+			AssetToolsModule.Get().CreateUniqueAssetName(ParentBlueprint->GetOutermost()->GetName(), TEXT("_Child"), PackageName, Name);
+			const FString PackagePath = FPackageName::GetLongPackagePath(PackageName);
+
+			const UAssetDefinition_Blueprint* BlueprintAssetDefinition = Cast<UAssetDefinition_Blueprint>(UAssetDefinitionRegistry::Get()->GetAssetDefinitionForClass(ParentBlueprint->GetClass()));
+			if (BlueprintAssetDefinition)
+			{
+				UFactory* Factory = BlueprintAssetDefinition->GetFactoryForBlueprintType(ParentBlueprint);
+
+				FContentBrowserModule& ContentBrowserModule = FModuleManager::LoadModuleChecked<FContentBrowserModule>("ContentBrowser");
+				ContentBrowserModule.Get().CreateNewAsset(Name, PackagePath, ParentBlueprint->GetClass(), Factory);
+			}
+		}
+	}
+
+	static void ExecuteEditDefaults(const FToolMenuContext& MenuContext, TArray<FAssetData> BlueprintAssets)
+	{
+		TArray<UBlueprint*> Blueprints;
+
+		FMessageLog EditorErrors("EditorErrors");
+		EditorErrors.NewPage(LOCTEXT("ExecuteEditDefaultsNewLogPage", "Loading Blueprints"));
+
+		for (const FAssetData& BlueprintAsset : BlueprintAssets)
+		{
+			if (UBlueprint* Object = Cast<UBlueprint>(BlueprintAsset.GetAsset()))
+			{
+				// If the blueprint is valid, allow it to be added to the list, otherwise log the error.
+				if ( Object->SkeletonGeneratedClass && Object->GeneratedClass )
+				{
+					Blueprints.Add(Object);
+				}
+				else
+				{
+					FFormatNamedArguments Arguments;
+					Arguments.Add(TEXT("ObjectName"), FText::FromString(Object->GetName()));
+					EditorErrors.Error(FText::Format(LOCTEXT("LoadBlueprint_FailedLog", "{ObjectName} could not be loaded because it derives from an invalid class.  Check to make sure the parent class for this blueprint hasn't been removed!"), Arguments ) );
+				}
+			}
+		}
+
+		if ( Blueprints.Num() > 0 )
+		{
+			FBlueprintEditorModule& BlueprintEditorModule = FModuleManager::LoadModuleChecked<FBlueprintEditorModule>( "Kismet" );
+			TSharedRef< IBlueprintEditor > NewBlueprintEditor = BlueprintEditorModule.CreateBlueprintEditor(  EToolkitMode::Standalone, TSharedPtr<IToolkitHost>(), Blueprints );
+		}
+
+		// Report errors
+		EditorErrors.Notify(LOCTEXT("OpenDefaults_Failed", "Opening Class Defaults Failed!"));
+	}
+	
+	static FDelayedAutoRegisterHelper DelayedAutoRegister(EDelayedRegisterRunPhase::EndOfEngineInit, []{ 
+		UToolMenus::RegisterStartupCallback(FSimpleMulticastDelegate::FDelegate::CreateLambda([]()
+		{
+			FToolMenuOwnerScoped OwnerScoped(UE_MODULE_NAME);
+			UToolMenu* Menu = UE::ContentBrowser::ExtendToolMenu_AssetContextMenu(UBlueprint::StaticClass());
+	        
+			FToolMenuSection& Section = Menu->FindOrAddSection("GetAssetActions");
+				Section.AddDynamicEntry(NAME_None, FNewToolMenuSectionDelegate::CreateLambda([](FToolMenuSection& InSection)
+				{
+					if (const UContentBrowserAssetContextMenuContext* Context = UContentBrowserAssetContextMenuContext::FindContextWithAssets(InSection))
+					{
+						// TODO NDarnell re: EIncludeSubclasses::No, Temporary - Need to ensure we don't have duplicates for now, because no all subclasses of blueprint are of this class yet.
+						if (const FAssetData* SelectedBlueprintPtr = Context->GetSingleSelectedAssetOfType(UBlueprint::StaticClass(), EIncludeSubclasses::No))
+						{
+							const TAttribute<FText> Label = LOCTEXT("Blueprint_NewDerivedBlueprint", "Create Child Blueprint Class");
+							const TAttribute<FText> ToolTip = TAttribute<FText>::CreateLambda([SelectedBlueprintPtr]()
+							{
+								const uint32 BPFlags = SelectedBlueprintPtr->GetTagValueRef<uint32>(FBlueprintTags::ClassFlags);
+								if ((BPFlags & (CLASS_Deprecated)) == 0)
+								{
+									return LOCTEXT("Blueprint_NewDerivedBlueprintTooltip", "Creates a Child Blueprint Class based on the current Blueprint, allowing you to create variants easily.");
+								}
+								else
+								{
+									return LOCTEXT("Blueprint_NewDerivedBlueprintIsDeprecatedTooltip", "Blueprint class is deprecated, cannot derive a child Blueprint!");
+								}
+							});
+							const FSlateIcon Icon = FSlateIcon(FAppStyle::GetAppStyleSetName(), "ClassIcon.Blueprint");
+
+							FToolUIAction DeriveNewBlueprint;
+							DeriveNewBlueprint.ExecuteAction = FToolMenuExecuteAction::CreateStatic(&ExecuteNewDerivedBlueprint, SelectedBlueprintPtr);
+							DeriveNewBlueprint.CanExecuteAction = FToolMenuCanExecuteAction::CreateStatic(&CanExecuteNewDerivedBlueprint, SelectedBlueprintPtr);
+							InSection.AddMenuEntry("CreateChildBlueprintClass", Label, ToolTip, Icon, DeriveNewBlueprint);
+						}
+
+						TArray<FAssetData> SelectedBlueprints = Context->GetSelectedAssetsOfType(UBlueprint::StaticClass(), EIncludeSubclasses::No);
+						if (SelectedBlueprints.Num() > 1)
+						{
+							TArray<UClass*> SelectedBlueprintParentClasses;
+							Algo::Transform(SelectedBlueprints, SelectedBlueprintParentClasses, [](const FAssetData& BlueprintAsset){ return UBlueprint::GetBlueprintParentClassFromAssetTags(BlueprintAsset); });
+							
+							// Ensure that all the selected blueprints are actors
+							const bool bAreAllSelectedBlueprintsActors =
+								Algo::AllOf(SelectedBlueprintParentClasses, [](UClass* ParentClass){ return ParentClass->IsChildOf(AActor::StaticClass()); });
+							
+							if (bAreAllSelectedBlueprintsActors)
+							{
+								const TAttribute<FText> Label = LOCTEXT("Blueprint_EditDefaults", "Edit Shared Defaults");
+								const TAttribute<FText> ToolTip = LOCTEXT("Blueprint_EditDefaultsTooltip", "Edit the shared default properties of the selected actor blueprints.");
+								const FSlateIcon Icon = FSlateIcon(FAppStyle::GetAppStyleSetName(), "Kismet.Tabs.BlueprintDefaults");
+								const FToolMenuExecuteAction UIAction = FToolMenuExecuteAction::CreateStatic(&ExecuteEditDefaults, MoveTemp(SelectedBlueprints));
+								InSection.AddMenuEntry("Blueprint_EditDefaults", Label, ToolTip, Icon, UIAction);
+							}
+						}
+					}
+				}));
+		}));
+	});
+}
+
+#undef LOCTEXT_NAMESPACE
diff --git a/Engine/Plugins/Editor/EngineAssetDefinitions/Source/Private/Script/AssetDefinition_Class.cpp b/Engine/Plugins/Editor/EngineAssetDefinitions/Source/Private/Script/AssetDefinition_Class.cpp
new file mode 100644
--- /dev/null
+++ b/Engine/Plugins/Editor/EngineAssetDefinitions/Source/Private/Script/AssetDefinition_Class.cpp
@@ -0,0 +1,182 @@
+// Copyright Epic Games, Inc. All Rights Reserved.
+
+#include "AssetDefinition_Class.h"
+#include "ToolMenus.h"
+#include "AssetRegistry/AssetData.h"
+#include "HAL/FileManager.h"
+#include "Styling/AppStyle.h"
+#include "SourceCodeNavigation.h"
+#include "AddToProjectConfig.h"
+#include "GameProjectGenerationModule.h"
+#include "Kismet2/KismetEditorUtilities.h"
+#include "Editor/UnrealEdEngine.h"
+#include "Preferences/UnrealEdOptions.h"
+#include "UnrealEdGlobals.h"
+
+#define LOCTEXT_NAMESPACE "AssetTypeActions"
+
+void FAssetTypeActions_Class::GetActions(const TArray<UObject*>& InObjects, FToolMenuSection& Section)
+{
+	UClass *const BaseClass = (InObjects.Num() == 1) ? Cast<UClass>(InObjects[0]) : nullptr;
+
+	// Only allow the New class option if we have a base class that we can actually derive from in one of our project modules
+	FGameProjectGenerationModule& GameProjectGenerationModule = FGameProjectGenerationModule::Get();
+	TArray<FModuleContextInfo> ProjectModules = GameProjectGenerationModule.GetCurrentProjectModules();
+	const bool bIsValidBaseCppClass = BaseClass && GameProjectGenerationModule.IsValidBaseClassForCreation(BaseClass, ProjectModules);
+	const bool bIsValidBaseBlueprintClass = BaseClass && FKismetEditorUtilities::CanCreateBlueprintOfClass(BaseClass);
+
+	auto CreateCreateDerivedCppClass = [BaseClass]()
+	{
+		// Work out where the header file for the current class is, as we'll use that path as the default for the new class
+		FString BaseClassPath;
+		if(FSourceCodeNavigation::FindClassHeaderPath(BaseClass, BaseClassPath))
+		{
+			// Strip off the actual filename as we only need the path
+			BaseClassPath = FPaths::GetPath(BaseClassPath);
+		}
+
+		FGameProjectGenerationModule::Get().OpenAddCodeToProjectDialog(
+			FAddToProjectConfig()
+			.ParentClass(BaseClass)
+			.InitialPath(BaseClassPath)
+			.ParentWindow(FGlobalTabmanager::Get()->GetRootWindow())
+		);
+	};
+
+	auto IsCPPAllowed = []()
+	{
+		return ensure(GUnrealEd) && GUnrealEd->GetUnrealEdOptions()->IsCPPAllowed();
+	};
+
+	auto CanCreateDerivedCppClass = [bIsValidBaseCppClass, IsCPPAllowed]() -> bool
+	{
+		return bIsValidBaseCppClass && IsCPPAllowed();
+	};
+
+	auto CreateCreateDerivedBlueprintClass = [BaseClass]()
+	{
+		FGameProjectGenerationModule::Get().OpenAddBlueprintToProjectDialog(
+			FAddToProjectConfig()
+			.ParentClass(BaseClass)
+			.ParentWindow(FGlobalTabmanager::Get()->GetRootWindow())
+		);
+	};
+
+	auto CanCreateDerivedBlueprintClass = [bIsValidBaseBlueprintClass]() -> bool
+	{
+		return bIsValidBaseBlueprintClass;
+	};
+
+	FText NewDerivedCppClassLabel;
+	FText NewDerivedCppClassToolTip;
+	FText NewDerivedBlueprintClassLabel;
+	FText NewDerivedBlueprintClassToolTip;
+	if(InObjects.Num() == 1)
+	{
+		check(BaseClass);
+
+		const FText BaseClassName = FText::FromName(BaseClass->GetFName());
+
+		NewDerivedCppClassLabel = FText::Format(LOCTEXT("Class_NewDerivedCppClassLabel_CreateFrom", "Create C++ class derived from {0}"), BaseClassName);
+		if(bIsValidBaseCppClass)
+		{
+			NewDerivedCppClassToolTip = FText::Format(LOCTEXT("Class_NewDerivedCppClassTooltip_CreateFrom", "Create a new C++ class deriving from {0}."), BaseClassName);
+		}
+		else
+		{
+			NewDerivedCppClassToolTip = FText::Format(LOCTEXT("Class_NewDerivedCppClassTooltip_InvalidClass", "Cannot create a new C++ class deriving from {0}."), BaseClassName);
+		}
+
+		NewDerivedBlueprintClassLabel = FText::Format(LOCTEXT("Class_NewDerivedBlueprintClassLabel_CreateFrom", "Create Blueprint class based on {0}"), BaseClassName);
+		if(bIsValidBaseBlueprintClass)
+		{
+			NewDerivedBlueprintClassToolTip = FText::Format(LOCTEXT("Class_NewDerivedBlueprintClassTooltip_CreateFrom", "Create a new Blueprint class based on {0}."), BaseClassName);
+		}
+		else
+		{
+			NewDerivedBlueprintClassToolTip = FText::Format(LOCTEXT("Class_NewDerivedBlueprintClassTooltip_InvalidClass", "Cannot create a new Blueprint class based on {0}."), BaseClassName);
+		}
+	}
+	else
+	{
+		NewDerivedCppClassLabel = LOCTEXT("Class_NewDerivedCppClassLabel_InvalidNumberOfBases", "New C++ class derived from...");
+		NewDerivedCppClassToolTip = LOCTEXT("Class_NewDerivedCppClassTooltip_InvalidNumberOfBases", "Can only create a derived C++ class when there is a single base class selected.");
+
+		NewDerivedBlueprintClassLabel = LOCTEXT("Class_NewDerivedBlueprintClassLabel_InvalidNumberOfBases", "New Blueprint class based on...");
+		NewDerivedBlueprintClassToolTip = LOCTEXT("Class_NewDerivedBlueprintClassTooltip_InvalidNumberOfBases", "Can only create a Blueprint class when there is a single base class selected.");
+	}
+
+	Section.AddMenuEntry(
+		"NewDerivedCppClass",
+		NewDerivedCppClassLabel,
+		NewDerivedCppClassToolTip,
+		FSlateIcon(FAppStyle::GetAppStyleSetName(), "Icons.C++"),
+		FUIAction(
+			FExecuteAction::CreateLambda(CreateCreateDerivedCppClass),
+			FCanExecuteAction::CreateLambda(CanCreateDerivedCppClass),
+			FIsActionChecked(),
+			FIsActionButtonVisible::CreateLambda(IsCPPAllowed)
+		)
+	);
+
+	Section.AddMenuEntry(
+		"NewDerivedBlueprintClass",
+		NewDerivedBlueprintClassLabel,
+		NewDerivedBlueprintClassToolTip,
+		FSlateIcon(FAppStyle::GetAppStyleSetName(), "Icons.Blueprint"),
+		FUIAction(
+			FExecuteAction::CreateLambda(CreateCreateDerivedBlueprintClass),
+			FCanExecuteAction::CreateLambda(CanCreateDerivedBlueprintClass)
+			)
+		);
+}
+
+UThumbnailInfo* FAssetTypeActions_Class::GetThumbnailInfo(UObject* Asset) const
+{
+	// todo: jdale - CLASS - We need to generate and store proper thumbnail info for classes so that we can store their custom render transforms
+	// This can't be stored in the UClass instance (like we do for Blueprints), so we'll need another place to store it
+	// This will need to be accessible to FClassThumbnailScene::GetSceneThumbnailInfo
+	return nullptr;
+}
+
+void FAssetTypeActions_Class::OpenAssetEditor( const TArray<UObject*>& InObjects, TSharedPtr<class IToolkitHost> EditWithinLevelEditor )
+{
+	for(UObject* Object : InObjects)
+	{
+		UClass* const Class = Cast<UClass>(Object);
+		if(Class)
+		{
+			TArray<FString> FilesToOpen;
+
+			FString ClassHeaderPath;
+			if(FSourceCodeNavigation::FindClassHeaderPath(Class, ClassHeaderPath))
+			{
+				const FString AbsoluteHeaderPath = IFileManager::Get().ConvertToAbsolutePathForExternalAppForRead(*ClassHeaderPath);
+				FilesToOpen.Add(AbsoluteHeaderPath);
+			}
+
+			FString ClassSourcePath;
+			if(FSourceCodeNavigation::FindClassSourcePath(Class, ClassSourcePath))
+			{
+				const FString AbsoluteSourcePath = IFileManager::Get().ConvertToAbsolutePathForExternalAppForRead(*ClassSourcePath);
+				FilesToOpen.Add(AbsoluteSourcePath);
+			}
+
+			FSourceCodeNavigation::OpenSourceFiles(FilesToOpen);
+		}
+	}
+}
+
+TWeakPtr<IClassTypeActions> FAssetTypeActions_Class::GetClassTypeActions(const FAssetData& AssetData) const
+{
+	UClass* Class = FindObject<UClass>(nullptr, *AssetData.GetObjectPathString());
+	if(Class)
+	{
+		FAssetToolsModule& AssetToolsModule = FAssetToolsModule::GetModule();
+		return AssetToolsModule.Get().GetClassTypeActionsForClass(Class);
+	}
+
+	return nullptr;
+}
+
+#undef LOCTEXT_NAMESPACE
diff --git a/Engine/Plugins/Editor/EngineAssetDefinitions/Source/Private/Script/AssetDefinition_Class.h b/Engine/Plugins/Editor/EngineAssetDefinitions/Source/Private/Script/AssetDefinition_Class.h
new file mode 100644
--- /dev/null
+++ b/Engine/Plugins/Editor/EngineAssetDefinitions/Source/Private/Script/AssetDefinition_Class.h
@@ -0,0 +1,30 @@
+// Copyright Epic Games, Inc. All Rights Reserved.
+
+#pragma once
+
+#include "CoreMinimal.h"
+#include "UObject/Class.h"
+#include "Toolkits/IToolkitHost.h"
+#include "AssetTypeActions/AssetTypeActions_ClassTypeBase.h"
+
+struct FAssetData;
+class IClassTypeActions;
+
+class FAssetTypeActions_Class : public FAssetTypeActions_ClassTypeBase
+{
+public:
+	// IAssetTypeActions Implementation
+	virtual FText GetName() const override { return NSLOCTEXT("AssetTypeActions", "AssetTypeActions_Class", "C++ Class"); }
+	virtual FColor GetTypeColor() const override { return FColor(255, 255, 255); }
+	virtual UClass* GetSupportedClass() const override { return UClass::StaticClass(); }
+	virtual uint32 GetCategories() override { return EAssetTypeCategories::Basic; }
+	
+	virtual void GetActions(const TArray<UObject*>& InObjects, struct FToolMenuSection& Section) override;
+
+	virtual class UThumbnailInfo* GetThumbnailInfo(UObject* Asset) const override;
+
+	virtual void OpenAssetEditor( const TArray<UObject*>& InObjects, TSharedPtr<class IToolkitHost> EditWithinLevelEditor = TSharedPtr<IToolkitHost>() ) override;
+
+	// FAssetTypeActions_ClassTypeBase Implementation
+	virtual TWeakPtr<IClassTypeActions> GetClassTypeActions(const FAssetData& AssetData) const override;
+};
diff --git a/Engine/Plugins/Editor/EngineAssetDefinitions/Source/Private/Script/AssetDefinition_ClassTypeBase.cpp b/Engine/Plugins/Editor/EngineAssetDefinitions/Source/Private/Script/AssetDefinition_ClassTypeBase.cpp
new file mode 100644
--- /dev/null
+++ b/Engine/Plugins/Editor/EngineAssetDefinitions/Source/Private/Script/AssetDefinition_ClassTypeBase.cpp
@@ -0,0 +1,19 @@
+// Copyright Epic Games, Inc. All Rights Reserved.
+
+#include "Script/AssetDefinition_ClassTypeBase.h"
+#include "IClassTypeActions.h"
+
+#define LOCTEXT_NAMESPACE "AssetTypeActions"
+
+TSharedPtr<SWidget> UAssetDefinition_ClassTypeBase::GetThumbnailOverlay(const FAssetData& AssetData) const
+{
+	TSharedPtr<IClassTypeActions> ClassTypeActions = GetClassTypeActions(AssetData).Pin();
+	if (ClassTypeActions.IsValid())
+	{
+		return ClassTypeActions->GetThumbnailOverlay(AssetData);
+	}
+
+	return nullptr;
+}
+
+#undef LOCTEXT_NAMESPACE
diff --git a/Engine/Plugins/Editor/EngineAssetDefinitions/Source/Public/AssetDefinition_Blueprint.h b/Engine/Plugins/Editor/EngineAssetDefinitions/Source/Public/AssetDefinition_Blueprint.h
deleted file mode 100644
diff --git a/Engine/Plugins/Editor/EngineAssetDefinitions/Source/Public/Script/AssetDefinition_Blueprint.h b/Engine/Plugins/Editor/EngineAssetDefinitions/Source/Public/Script/AssetDefinition_Blueprint.h
new file mode 100644
--- /dev/null
+++ b/Engine/Plugins/Editor/EngineAssetDefinitions/Source/Public/Script/AssetDefinition_Blueprint.h
@@ -0,0 +1,54 @@
+// Copyright Epic Games, Inc. All Rights Reserved.
+
+#pragma once
+
+#include "AssetDefinition_ClassTypeBase.h"
+#include "Engine/Blueprint.h"
+
+#include "AssetDefinition_Blueprint.generated.h"
+
+struct FToolMenuContext;
+struct FAssetData;
+class UFactory;
+
+UCLASS()
+class UAssetDefinition_Blueprint : public UAssetDefinition_ClassTypeBase
+{
+	GENERATED_BODY()
+
+public:
+	// UAssetDefinition Implementation
+	virtual FText GetAssetDisplayName() const override { return NSLOCTEXT("AssetTypeActions", "AssetTypeActions_Blueprint", "Blueprint Class"); }
+	virtual FText GetAssetDisplayName(const FAssetData& AssetData) const override;
+	virtual FText GetAssetDescription(const FAssetData& AssetData) const override;
+	virtual FLinearColor GetAssetColor() const override { return FLinearColor(FColor( 63, 126, 255 )); }
+	virtual TSoftClassPtr<UObject> GetAssetClass() const override { return UBlueprint::StaticClass(); }
+	virtual TConstArrayView<FAssetCategoryPath> GetAssetCategories() const override
+	{
+		static const auto Categories = { EAssetCategoryPaths::Basic, EAssetCategoryPaths::Blueprint };
+		return Categories;
+	}
+
+	virtual bool CanMerge() const override { return true; }
+	virtual EAssetCommandResult Merge(const FAssetAutomaticMergeArgs& MergeArgs) const override;
+	virtual EAssetCommandResult Merge(const FAssetManualMergeArgs& MergeArgs) const override;
+	
+	virtual EAssetCommandResult PerformAssetDiff(const FAssetDiffArgs& DiffArgs) const override;
+
+	virtual void BuildFilters(TArray<FAssetFilterData>& OutFilters) const override;
+	
+	virtual UThumbnailInfo* LoadThumbnailInfo(const FAssetData& InAssetData) const override;
+	virtual EAssetCommandResult OpenAssets(const FAssetOpenArgs& OpenArgs) const override;
+	// UAssetDefinition End
+
+	// UAssetDefinition_ClassTypeBase Implementation
+	virtual TWeakPtr<IClassTypeActions> GetClassTypeActions(const FAssetData& AssetData) const override;
+	// UAssetDefinition_ClassTypeBase End
+	
+public:
+	virtual UFactory* GetFactoryForBlueprintType(UBlueprint* InBlueprint) const;
+};
+
+#if UE_ENABLE_INCLUDE_ORDER_DEPRECATED_IN_5_2
+#include "CoreMinimal.h"
+#endif
diff --git a/Engine/Plugins/Editor/EngineAssetDefinitions/Source/Public/Script/AssetDefinition_ClassTypeBase.h b/Engine/Plugins/Editor/EngineAssetDefinitions/Source/Public/Script/AssetDefinition_ClassTypeBase.h
new file mode 100644
--- /dev/null
+++ b/Engine/Plugins/Editor/EngineAssetDefinitions/Source/Public/Script/AssetDefinition_ClassTypeBase.h
@@ -0,0 +1,31 @@
+// Copyright Epic Games, Inc. All Rights Reserved.
+
+#pragma once
+
+#include "AssetDefinitionDefault.h"
+
+#include "AssetDefinition_ClassTypeBase.generated.h"
+
+class IClassTypeActions;
+
+/** Base class for "class type" assets (C++ classes and Blueprints */
+UCLASS(Abstract)
+class UAssetDefinition_ClassTypeBase : public UAssetDefinitionDefault
+{
+	GENERATED_BODY()
+
+public:
+	// UAssetDefinition Begin
+	virtual FAssetSupportResponse CanLocalize(const FAssetData& InAsset) const override
+	{
+		return FAssetSupportResponse::NotSupported();
+	}
+	virtual TSharedPtr<class SWidget> GetThumbnailOverlay(const FAssetData& AssetData) const override;
+	// UAssetDefinition End
+
+	/** Get the class type actions for this asset */
+	virtual TWeakPtr<IClassTypeActions> GetClassTypeActions(const FAssetData& AssetData) const
+	{
+		return TWeakPtr<IClassTypeActions>();
+	}
+};
diff --git a/Engine/Source/Developer/AssetTools/Private/AssetDefinition_AssetTypeActionsProxy.cpp b/Engine/Source/Developer/AssetTools/Private/AssetDefinition_AssetTypeActionsProxy.cpp
--- a/Engine/Source/Developer/AssetTools/Private/AssetDefinition_AssetTypeActionsProxy.cpp
+++ b/Engine/Source/Developer/AssetTools/Private/AssetDefinition_AssetTypeActionsProxy.cpp
@@ -122,25 +122,25 @@ bool UAssetDefinition_AssetTypeActionsProxy::CanMerge() const
 	return AssetType->CanMerge();
 }
 
-EAssetCommandResult UAssetDefinition_AssetTypeActionsProxy::Merge(const FAssetMergeArgs& MergeArgs) const
+EAssetCommandResult UAssetDefinition_AssetTypeActionsProxy::Merge(const FAssetAutomaticMergeArgs& MergeArgs) const
 {
-	if (!MergeArgs.BaseAsset.IsSet() && !MergeArgs.RemoteAsset.IsSet())
-	{
-		AssetType->Merge(MergeArgs.LocalAsset.GetAsset());
-		return EAssetCommandResult::Handled;
-	}
-	else
-	{
-		FOnAssetMergeResolved MergeCallback = MergeArgs.ResolutionCallback;
-		AssetType->Merge(MergeArgs.BaseAsset.GetValue().GetAsset(), MergeArgs.RemoteAsset.GetValue().GetAsset(), MergeArgs.LocalAsset.GetAsset(),
-			FOnMergeResolved::CreateLambda([MergeCallback](UPackage* MergedPackage, EMergeResult::Type Result){
-			FAssetMergeResults Results;
-			Results.MergedPackage = MergedPackage;
-			Results.Result = static_cast<EAssetMergeResult>(Result);
-			MergeCallback.ExecuteIfBound(Results);
-		}));
-		return EAssetCommandResult::Handled;
-	}
+	AssetType->Merge(MergeArgs.LocalAsset);
+
+	return EAssetCommandResult::Handled;
+}
+
+EAssetCommandResult UAssetDefinition_AssetTypeActionsProxy::Merge(const FAssetManualMergeArgs& MergeArgs) const
+{
+	FOnAssetMergeResolved MergeCallback = MergeArgs.ResolutionCallback;
+	AssetType->Merge(MergeArgs.BaseAsset, MergeArgs.RemoteAsset, MergeArgs.LocalAsset,
+		FOnMergeResolved::CreateLambda([MergeCallback](UPackage* MergedPackage, EMergeResult::Type Result){
+		FAssetMergeResults Results;
+		Results.MergedPackage = MergedPackage;
+		Results.Result = static_cast<EAssetMergeResult>(Result);
+		MergeCallback.ExecuteIfBound(Results);
+	}));
+
+	return EAssetCommandResult::Handled;
 }
 
 FAssetOpenSupport UAssetDefinition_AssetTypeActionsProxy::GetAssetOpenSupport(const FAssetOpenSupportArgs& OpenSupportArgs) const
@@ -189,7 +189,7 @@ EAssetCommandResult UAssetDefinition_AssetTypeActionsProxy::OpenAssets(const FAs
 	return EAssetCommandResult::Handled;
 }
 
-EAssetCommandResult UAssetDefinition_AssetTypeActionsProxy::GetFilters(TArray<FAssetFilterData>& OutFilters) const
+void UAssetDefinition_AssetTypeActionsProxy::BuildFilters(TArray<FAssetFilterData>& OutFilters) const
 {
 	if (AssetType->CanFilter())
 	{
@@ -198,11 +198,7 @@ EAssetCommandResult UAssetDefinition_AssetTypeActionsProxy::GetFilters(TArray<FA
 		Data.DisplayText = AssetType->GetName();
 		AssetType->BuildBackendFilter(Data.Filter);
 		OutFilters.Add(MoveTemp(Data));
-
-		return EAssetCommandResult::Handled;
 	}
-
-	return EAssetCommandResult::Unhandled;
 }
 
 FText UAssetDefinition_AssetTypeActionsProxy::GetObjectDisplayNameText(UObject* Object) const
diff --git a/Engine/Source/Developer/AssetTools/Private/AssetDefinition_AssetTypeActionsProxy.h b/Engine/Source/Developer/AssetTools/Private/AssetDefinition_AssetTypeActionsProxy.h
--- a/Engine/Source/Developer/AssetTools/Private/AssetDefinition_AssetTypeActionsProxy.h
+++ b/Engine/Source/Developer/AssetTools/Private/AssetDefinition_AssetTypeActionsProxy.h
@@ -31,14 +31,15 @@ public:
 	virtual FAssetSupportResponse CanLocalize(const FAssetData& InAsset) const override;
 	virtual bool CanImport() const override;
 	virtual bool CanMerge() const override;
-	virtual EAssetCommandResult Merge(const FAssetMergeArgs& MergeArgs) const override;
+	virtual EAssetCommandResult Merge(const FAssetAutomaticMergeArgs& MergeArgs) const override;
+	virtual EAssetCommandResult Merge(const FAssetManualMergeArgs& MergeArgs) const override;
 
 	virtual FAssetOpenSupport GetAssetOpenSupport(const FAssetOpenSupportArgs& OpenSupportArgs) const override;
 	virtual TArray<FAssetData> PrepareToActivateAssets(const FAssetActivateArgs& ActivateArgs) const override;
 	virtual EAssetCommandResult ActivateAssets(const FAssetActivateArgs& ActivateArgs) const override;
 	virtual EAssetCommandResult OpenAssets(const FAssetOpenArgs& OpenArgs) const override;
 
-	virtual EAssetCommandResult GetFilters(TArray<FAssetFilterData>& OutFilters) const override;
+	virtual void BuildFilters(TArray<FAssetFilterData>& OutFilters) const override;
 
 	// We don't need to implement GetSourceFiles, the default one will / should work for any existing type.
 	// virtual EAssetCommandResult GetSourceFiles(const FAssetSourceFileArgs& SourceFileArgs, TArray<FAssetSourceFile>& OutSourceAssets) const override;
diff --git a/Engine/Source/Developer/AssetTools/Private/AssetTools.cpp b/Engine/Source/Developer/AssetTools/Private/AssetTools.cpp
--- a/Engine/Source/Developer/AssetTools/Private/AssetTools.cpp
+++ b/Engine/Source/Developer/AssetTools/Private/AssetTools.cpp
@@ -35,7 +35,6 @@
 #include "IClassTypeActions.h"
 #include "ClassViewerModule.h"
 #include "ClassViewerFilter.h"
-#include "AssetTypeActions/AssetTypeActions_Blueprint.h"
 #include "AssetTypeActions/AssetTypeActions_BlueprintGeneratedClass.h"
 #include "AssetTypeActions/AssetTypeActions_AnimationAsset.h"
 #include "AssetTypeActions/AssetTypeActions_AnimBlueprint.h"
@@ -43,7 +42,6 @@
 #include "AssetTypeActions/AssetTypeActions_AnimBlueprintGeneratedClass.h"
 #include "AssetTypeActions/AssetTypeActions_AnimBoneCompressionSettings.h"
 #include "AssetTypeActions/AssetTypeActions_AnimCurveCompressionSettings.h"
-#include "AssetTypeActions/AssetTypeActions_Class.h"
 #include "AssetTypeActions/AssetTypeActions_ForceFeedbackEffect.h"
 #include "AssetTypeActions/AssetTypeActions_ParticleSystem.h"
 #include "AssetTypeActions/AssetTypeActions_PhysicalMaterialMask.h"
@@ -279,25 +277,22 @@ public:
 
 	virtual bool CanFilter() override
 	{
-		TArray<FAssetFilterData> Filters;
-		AssetDefinitionPtr->GetFilters(Filters);
-		return Filters.Num() > 0;
+		TSharedRef<FAssetFilterDataCache> FilterCache = AssetDefinitionPtr->GetFilters();
+		return FilterCache->Filters.Num() > 0;
 	}
 
 	virtual FName GetFilterName() const override
 	{
-		TArray<FAssetFilterData> Filters;
-		AssetDefinitionPtr->GetFilters(Filters);
-		return Filters.Num() > 0 ? FName(*Filters[0].Name) : NAME_None;
+		TSharedRef<FAssetFilterDataCache> FilterCache = AssetDefinitionPtr->GetFilters();
+		return FilterCache->Filters.Num() > 0 ? FName(*FilterCache->Filters[0].Name) : NAME_None;
 	}
 
 	virtual void BuildBackendFilter(FARFilter& InFilter) override
 	{
-		TArray<FAssetFilterData> Filters;
-		AssetDefinitionPtr->GetFilters(Filters);
-		if (Filters.Num() > 0)
+		TSharedRef<FAssetFilterDataCache> FilterCache = AssetDefinitionPtr->GetFilters();
+		if (FilterCache->Filters.Num() > 0)
 		{
-			InFilter = Filters[0].Filter;
+			InFilter = FilterCache->Filters[0].Filter;
 		}
 	}
 
@@ -329,7 +324,7 @@ public:
 	/** Begins a merge operation for InObject (automatically determines remote/base versions needed to resolve) */
 	virtual void Merge( UObject* InObject ) override
 	{
-		FAssetMergeArgs MergeArgs;
+		FAssetAutomaticMergeArgs MergeArgs;
 		MergeArgs.LocalAsset = InObject;
 
 		AssetDefinitionPtr.Get()->Merge(MergeArgs);
@@ -338,7 +333,7 @@ public:
 	/** Begins a merge between the specified assets */
 	virtual void Merge(UObject* BaseAsset, UObject* RemoteAsset, UObject* LocalAsset, const FOnMergeResolved& ResolutionCallback) override
 	{
-		FAssetMergeArgs MergeArgs;
+		FAssetManualMergeArgs MergeArgs;
 		MergeArgs.LocalAsset = LocalAsset;
 		MergeArgs.BaseAsset = BaseAsset;
 		MergeArgs.RemoteAsset = RemoteAsset;
@@ -1334,9 +1329,7 @@ UAssetToolsImpl::UAssetToolsImpl(const FObjectInitializer& ObjectInitializer)
 	RegisterAssetTypeActions(MakeShareable(new FAssetTypeActions_AnimBlueprintGeneratedClass));
 	RegisterAssetTypeActions(MakeShareable(new FAssetTypeActions_AnimBoneCompressionSettings));
 	RegisterAssetTypeActions(MakeShareable(new FAssetTypeActions_AnimCurveCompressionSettings));
-	RegisterAssetTypeActions(MakeShareable(new FAssetTypeActions_Blueprint));
 	RegisterAssetTypeActions(MakeShareable(new FAssetTypeActions_BlueprintGeneratedClass));
-	RegisterAssetTypeActions(MakeShareable(new FAssetTypeActions_Class));
 	RegisterAssetTypeActions(MakeShareable(new FAssetTypeActions_ForceFeedbackEffect(InputCategoryBit)));
 	RegisterAssetTypeActions(MakeShareable(new FAssetTypeActions_ParticleSystem));
 	RegisterAssetTypeActions(MakeShareable(new FAssetTypeActions_PhysicalMaterialMask));
diff --git a/Engine/Source/Developer/AssetTools/Public/AssetTypeActions/AssetTypeActions_Blueprint.h b/Engine/Source/Developer/AssetTools/Public/AssetTypeActions/AssetTypeActions_Blueprint.h
--- a/Engine/Source/Developer/AssetTools/Public/AssetTypeActions/AssetTypeActions_Blueprint.h
+++ b/Engine/Source/Developer/AssetTools/Public/AssetTypeActions/AssetTypeActions_Blueprint.h
@@ -11,7 +11,9 @@ struct FAssetData;
 class IClassTypeActions;
 class UFactory;
 
-class ASSETTOOLS_API FAssetTypeActions_Blueprint : public FAssetTypeActions_ClassTypeBase
+class
+// UE_DEPRECATED(5.2, "The AssetDefinition system is replacing AssetTypeActions and UAssetDefinition_Blueprint replaced this.  Please see the Conversion Guide in AssetDefinition.h")
+ASSETTOOLS_API FAssetTypeActions_Blueprint : public FAssetTypeActions_ClassTypeBase
 {
 public:
 	// IAssetTypeActions Implementation
diff --git a/Engine/Source/Developer/AssetTools/Public/AssetTypeActions/AssetTypeActions_ClassTypeBase.h b/Engine/Source/Developer/AssetTools/Public/AssetTypeActions/AssetTypeActions_ClassTypeBase.h
--- a/Engine/Source/Developer/AssetTools/Public/AssetTypeActions/AssetTypeActions_ClassTypeBase.h
+++ b/Engine/Source/Developer/AssetTools/Public/AssetTypeActions/AssetTypeActions_ClassTypeBase.h
@@ -10,14 +10,13 @@ struct FAssetData;
 class IClassTypeActions;
 
 /** Base class for "class type" assets (C++ classes and Blueprints */
-class ASSETTOOLS_API FAssetTypeActions_ClassTypeBase : public FAssetTypeActions_Base
+class
+// UE_DEPRECATED(5.2, "The AssetDefinition system is replacing AssetTypeActions and this was replaced by AssetDefinition_ClassTypeBase.  Please see the Conversion Guide in AssetDefinition.h")
+ASSETTOOLS_API FAssetTypeActions_ClassTypeBase : public FAssetTypeActions_Base
 {
 public:
 	virtual bool CanLocalize() const override { return false; }
 
 	/** Get the class type actions for this asset */
 	virtual TWeakPtr<IClassTypeActions> GetClassTypeActions(const FAssetData& AssetData) const = 0;
-
-	// IAssetTypeActions Implementation
-	virtual TSharedPtr<class SWidget> GetThumbnailOverlay(const FAssetData& AssetData) const override;
 };
diff --git a/Engine/Source/Developer/ToolWidgets/Private/SWidgetDrawer.cpp b/Engine/Source/Developer/ToolWidgets/Private/SWidgetDrawer.cpp
--- a/Engine/Source/Developer/ToolWidgets/Private/SWidgetDrawer.cpp
+++ b/Engine/Source/Developer/ToolWidgets/Private/SWidgetDrawer.cpp
@@ -385,7 +385,7 @@ void SWidgetDrawer::OnGlobalFocusChanging(const FFocusEvent& FocusEvent, const F
 				FWidgetPath MenuHostPath;
 
 				// See if the menu being opened is part of the content browser path and if so the menu should not be dismissed
-				FSlateApplication::Get().GeneratePathToWidgetUnchecked(MenuHost.ToSharedRef(), MenuHostPath, EVisibility::All);
+				FSlateApplication::Get().GeneratePathToWidgetUnchecked(MenuHost.ToSharedRef(), MenuHostPath, EVisibility::Visible);
 				if (!MenuHostPath.ContainsWidget(ActiveDrawerOverlayContent.Get()))
 				{
 					bShouldDismiss = true;
diff --git a/Engine/Source/Editor/AssetDefinition/Private/AssetDefinition.cpp b/Engine/Source/Editor/AssetDefinition/Private/AssetDefinition.cpp
--- a/Engine/Source/Editor/AssetDefinition/Private/AssetDefinition.cpp
+++ b/Engine/Source/Editor/AssetDefinition/Private/AssetDefinition.cpp
@@ -90,7 +90,7 @@ bool UAssetDefinition::CanRegisterStatically() const
 	return !GetClass()->HasAnyClassFlags(CLASS_Abstract);
 }
 
-EAssetCommandResult UAssetDefinition::GetFilters(TArray<FAssetFilterData>& OutFilters) const
+void UAssetDefinition::BuildFilters(TArray<FAssetFilterData>& OutFilters) const
 {
 	const TSoftClassPtr<UObject> AssetClassPtr = GetAssetClass();
 
@@ -103,15 +103,12 @@ EAssetCommandResult UAssetDefinition::GetFilters(TArray<FAssetFilterData>& OutFi
 			FAssetFilterData DefaultFilter;
 			DefaultFilter.Name = AssetClassPtr.ToSoftObjectPath().ToString();
 			DefaultFilter.DisplayText = GetAssetDisplayName();
+			DefaultFilter.FilterCategories = GetAssetCategories();
 			DefaultFilter.Filter.ClassPaths.Add(AssetClassPtr.ToSoftObjectPath().GetAssetPath());
 			DefaultFilter.Filter.bRecursiveClasses = true;
 			OutFilters.Add(MoveTemp(DefaultFilter));
-	
-			return EAssetCommandResult::Handled;
 		}
 	}
-	
-	return EAssetCommandResult::Unhandled;
 }
 
 #undef LOCTEXT_NAMESPACE
\ No newline at end of file
diff --git a/Engine/Source/Editor/AssetDefinition/Public/AssetDefinition.h b/Engine/Source/Editor/AssetDefinition/Public/AssetDefinition.h
--- a/Engine/Source/Editor/AssetDefinition/Public/AssetDefinition.h
+++ b/Engine/Source/Editor/AssetDefinition/Public/AssetDefinition.h
@@ -7,10 +7,10 @@
 #include "AssetRegistry/AssetData.h"
 #include "Misc/ScopedSlowTask.h"
 #include "Toolkits/IToolkit.h"
+#include "Misc/AssetFilterData.h"
 
 #include "AssetDefinition.generated.h"
 
-struct FAssetFilterData;
 class IToolkitHost;
 class UThumbnailInfo;
 struct FSlateBrush;
@@ -55,7 +55,7 @@ struct FAssetArgs
 	TConstArrayView<FAssetData> Assets;
 
 	template<typename ExpectedObjectType>
-	TArray<ExpectedObjectType*> LoadObjects(const TSet<FName>& LoadTags = {}) const
+	TArray<ExpectedObjectType*> LoadObjects(const TSet<FName>& LoadTags = {}, TArray<FAssetData>* OutAssetsThatFailedToLoad = nullptr) const
 	{
 		FScopedSlowTask SlowTask(Assets.Num());
 	
@@ -71,8 +71,15 @@ struct FAssetArgs
 				if (ExpectedObjectType* ExpectedType = Cast<ExpectedObjectType>(Asset.GetAsset(LoadTags)))
 				{
 					LoadedObjects.Add(ExpectedType);
+					continue;
 				}
-			}	
+			}
+
+			// If we get here than we failed to load the asset for some reason or another.
+			if (OutAssetsThatFailedToLoad)
+			{
+				OutAssetsThatFailedToLoad->Add(Asset);
+			}
 		}
 		
 		return LoadedObjects;
@@ -123,12 +130,17 @@ struct FAssetMergeResults
 
 DECLARE_DELEGATE_OneParam(FOnAssetMergeResolved, const FAssetMergeResults& Results);
 
-struct FAssetMergeArgs
+struct FAssetAutomaticMergeArgs
 {
-	FAssetData LocalAsset;
-	TOptional<FAssetData> BaseAsset;
-	TOptional<FAssetData> RemoteAsset;
-	
+	UObject* LocalAsset = nullptr;
+	FOnAssetMergeResolved ResolutionCallback;
+};
+
+struct FAssetManualMergeArgs
+{
+	UObject* LocalAsset = nullptr;
+	UObject* BaseAsset = nullptr;
+	UObject* RemoteAsset = nullptr;
 	FOnAssetMergeResolved ResolutionCallback;
 };
 
@@ -201,30 +213,6 @@ struct FAssetOpenSupportArgs
 	EAssetOpenMethod OpenMethod = EAssetOpenMethod::Edit;
 };
 
-/**
- * The asset category path is how we know how to build menus around assets.  For example, Basic is generally the ones
- * we expose at the top level, where as everything else is a category with a pull out menu, and the subcategory would
- * be where it gets placed in a submenu inside of there.
- */
-struct ASSETDEFINITION_API FAssetCategoryPath
-{
-	FAssetCategoryPath(const FText& InCategory);
-	FAssetCategoryPath(FText InCategory, FText InSubCategory);
-	FAssetCategoryPath(const FAssetCategoryPath& InCategory, const FText& SubCategory);
-	FAssetCategoryPath(TConstArrayView<FText> InCategoryPath);
-
-	FName GetCategory() const { return CategoryPath[0].Key; }
-	FText GetCategoryText() const { return CategoryPath[0].Value; }
-	
-	bool HasSubCategory() const { return CategoryPath.Num() > 1; }
-	FName GetSubCategory() const { return HasSubCategory() ? CategoryPath[1].Key : NAME_None; }
-	FText GetSubCategoryText() const { return HasSubCategory() ? CategoryPath[1].Value : FText::GetEmpty(); }
-	
-	FAssetCategoryPath operator / (const FText& SubCategory) const { return FAssetCategoryPath(*this, SubCategory); }
-	
-private:
-	TArray<TPair<FName, FText>> CategoryPath;
-};
 
 /**
  * These are just some common asset categories.  You're not at all limited to these, and can register an "Advanced"
@@ -282,6 +270,14 @@ enum class EIncludeClassInFilter : uint8
 	Always
 };
 
+
+struct FAssetFilterDataCache
+{
+public:
+	TArray<FAssetFilterData> Filters;
+};
+
+
 /**
  * Asset Definitions represent top level assets that are known to the editor.
  *
@@ -370,13 +366,26 @@ public:
 	
 	// Merging
 	virtual bool CanMerge() const { return false; }
-	virtual EAssetCommandResult Merge(const FAssetMergeArgs& MergeArgs) const { return EAssetCommandResult::Unhandled; }
-	
+	virtual EAssetCommandResult Merge(const FAssetAutomaticMergeArgs& MergeArgs) const { return EAssetCommandResult::Unhandled; }
+	virtual EAssetCommandResult Merge(const FAssetManualMergeArgs& MergeArgs) const { return EAssetCommandResult::Unhandled; }
+
 
 	// Filtering
-	virtual EAssetCommandResult GetFilters(TArray<FAssetFilterData>& OutFilters) const;
+	virtual TSharedRef<FAssetFilterDataCache> GetFilters() const
+	{
+		if (!FilterCache.IsValid())
+		{
+			FilterCache = MakeShared<FAssetFilterDataCache>();
+			BuildFilters(FilterCache->Filters);
+		}
 
-	
+		return FilterCache.ToSharedRef();
+	}
+
+protected:
+	virtual void BuildFilters(TArray<FAssetFilterData>& OutFilters) const;
+
+public:
 	// Extras
 	virtual FText GetObjectDisplayNameText(UObject* Object) const
 	{
@@ -444,5 +453,8 @@ protected:
 protected:
 	EIncludeClassInFilter IncludeClassInFilter = EIncludeClassInFilter::IfClassIsNotAbstract;
 
+private:
+	mutable TSharedPtr<FAssetFilterDataCache> FilterCache;
+
 	friend class UAssetDefinitionRegistry;
 };
\ No newline at end of file
diff --git a/Engine/Source/Editor/AssetDefinition/Public/Misc/AssetCategoryPath.h b/Engine/Source/Editor/AssetDefinition/Public/Misc/AssetCategoryPath.h
new file mode 100644
--- /dev/null
+++ b/Engine/Source/Editor/AssetDefinition/Public/Misc/AssetCategoryPath.h
@@ -0,0 +1,29 @@
+// Copyright Epic Games, Inc. All Rights Reserved.
+
+#pragma once
+
+#include "CoreMinimal.h"
+
+/**
+ * The asset category path is how we know how to build menus around assets.  For example, Basic is generally the ones
+ * we expose at the top level, where as everything else is a category with a pull out menu, and the subcategory would
+ * be where it gets placed in a submenu inside of there.
+ */
+struct ASSETDEFINITION_API FAssetCategoryPath
+{
+	FAssetCategoryPath(const FText& InCategory);
+	FAssetCategoryPath(const FAssetCategoryPath& InCategory, const FText& SubCategory);
+	FAssetCategoryPath(TConstArrayView<FText> InCategoryPath);
+
+	FName GetCategory() const { return CategoryPath[0].Key; }
+	FText GetCategoryText() const { return CategoryPath[0].Value; }
+	
+	bool HasSubCategory() const { return CategoryPath.Num() > 1; }
+	FName GetSubCategory() const { return HasSubCategory() ? CategoryPath[1].Key : NAME_None; }
+	FText GetSubCategoryText() const { return HasSubCategory() ? CategoryPath[1].Value : FText::GetEmpty(); }
+	
+	FAssetCategoryPath operator / (const FText& SubCategory) const { return FAssetCategoryPath(*this, SubCategory); }
+	
+private:
+	TArray<TPair<FName, FText>> CategoryPath;
+};
\ No newline at end of file
diff --git a/Engine/Source/Editor/AssetDefinition/Public/Misc/AssetFilterData.h b/Engine/Source/Editor/AssetDefinition/Public/Misc/AssetFilterData.h
--- a/Engine/Source/Editor/AssetDefinition/Public/Misc/AssetFilterData.h
+++ b/Engine/Source/Editor/AssetDefinition/Public/Misc/AssetFilterData.h
@@ -4,6 +4,7 @@
 
 #include "UObject/Object.h"
 #include "AssetRegistry/ARFilter.h"
+#include "AssetCategoryPath.h"
 
 #include "AssetFilterData.generated.h"
 
@@ -14,6 +15,10 @@ struct FAssetFilterData
 
 public:
 	FString Name;
-	FText DisplayText;	
+	
+	FText DisplayText;
+
+	TArray<FAssetCategoryPath> FilterCategories;
+	
 	FARFilter Filter;
 };
diff --git a/Engine/Source/Editor/ContentBrowser/Private/SFilterList.cpp b/Engine/Source/Editor/ContentBrowser/Private/SFilterList.cpp
--- a/Engine/Source/Editor/ContentBrowser/Private/SFilterList.cpp
+++ b/Engine/Source/Editor/ContentBrowser/Private/SFilterList.cpp
@@ -331,7 +331,7 @@ void SFilterList::PopulateAddFilterMenu_Internal(UToolMenu* Menu)
 		MenuExpansion = Context->MenuExpansion;
 	}
 
-	this->PopulateAddFilterMenu(Menu, AssetFilterCategories.FindChecked(EAssetTypeCategories::Basic), FOnFilterAssetType::CreateLambda([this](UClass *TestClass)
+	this->PopulateAddFilterMenu(Menu, AssetFilterCategories.FindChecked(EAssetCategoryPaths::Basic.GetCategory()), FOnFilterAssetType::CreateLambda([this](UClass *TestClass)
 	{
 		return !IsFilteredByPicker(this->InitialClassFilters, TestClass);
 	}));
diff --git a/Engine/Source/Editor/ContentBrowser/Private/SThumbnailEditModeTools.cpp b/Engine/Source/Editor/ContentBrowser/Private/SThumbnailEditModeTools.cpp
--- a/Engine/Source/Editor/ContentBrowser/Private/SThumbnailEditModeTools.cpp
+++ b/Engine/Source/Editor/ContentBrowser/Private/SThumbnailEditModeTools.cpp
@@ -52,14 +52,6 @@ void SThumbnailEditModeTools::Construct( const FArguments& InArgs, const TShared
 	DragStartLocation = FIntPoint(ForceInitToZero);
 	bInSmallView = InArgs._SmallView;
 
-	// Prime the SceneThumbnailInfo pointer
-	GetSceneThumbnailInfo();
-
-	if ( AssetThumbnail.IsValid() )
-	{
-		AssetThumbnail.Pin()->OnAssetDataChanged().AddSP(this, &SThumbnailEditModeTools::OnAssetDataChanged);
-	}
-
 	ChildSlot
 	[
 		SNew(SHorizontalBox)
@@ -104,13 +96,13 @@ void SThumbnailEditModeTools::Construct( const FArguments& InArgs, const TShared
 
 EVisibility SThumbnailEditModeTools::GetPrimitiveToolsVisibility() const
 {
-	const bool bIsVisible = !bInSmallView && (ConstGetSceneThumbnailInfoWithPrimitive() != NULL);
+	const bool bIsVisible = !bInSmallView && (GetSceneThumbnailInfo() != nullptr);
 	return bIsVisible ? EVisibility::Visible : EVisibility::Collapsed;
 }
 
 EVisibility SThumbnailEditModeTools::GetPrimitiveToolsResetToDefaultVisibility() const
 {
-	USceneThumbnailInfo* ThumbnailInfo = SceneThumbnailInfo.Get();
+	USceneThumbnailInfo* ThumbnailInfo = GetSceneThumbnailInfo();
 	
 	EVisibility ResetToDefaultVisibility = EVisibility::Collapsed;
 	if (ThumbnailInfo)
@@ -123,7 +115,7 @@ EVisibility SThumbnailEditModeTools::GetPrimitiveToolsResetToDefaultVisibility()
 
 const FSlateBrush* SThumbnailEditModeTools::GetCurrentPrimitiveBrush() const
 {
-	USceneThumbnailInfoWithPrimitive* ThumbnailInfo = ConstGetSceneThumbnailInfoWithPrimitive();
+	USceneThumbnailInfoWithPrimitive* ThumbnailInfo = GetSceneThumbnailInfoWithPrimitive();
 	if ( ThumbnailInfo )
 	{
 		// Note this is for the icon only.  we are assuming the thumbnail renderer does the right thing when rendering
@@ -211,7 +203,7 @@ FReply SThumbnailEditModeTools::OnMouseButtonDown( const FGeometry& MyGeometry,
 		if ( ThumbnailInfo )
 		{
 			FThumbnailRenderingInfo* RenderInfo = GUnrealEd->GetThumbnailManager()->GetRenderingInfo(Asset);
-			if (RenderInfo != NULL && RenderInfo->Renderer != NULL)
+			if (RenderInfo != nullptr && RenderInfo->Renderer != nullptr)
 			{
 				bModifiedThumbnailWhileDragging = false;
 				DragStartLocation = FIntPoint(MouseEvent.GetScreenSpacePosition().X, MouseEvent.GetScreenSpacePosition().Y);
@@ -314,15 +306,15 @@ FCursorReply SThumbnailEditModeTools::OnCursorQuery( const FGeometry& MyGeometry
 		FCursorReply::Cursor( EMouseCursor::Default );
 }
 
-USceneThumbnailInfo* SThumbnailEditModeTools::GetSceneThumbnailInfo()
+USceneThumbnailInfo* SThumbnailEditModeTools::GetSceneThumbnailInfo() const
 {
-	if ( !SceneThumbnailInfo.IsValid() )
+	USceneThumbnailInfo* SceneThumbnailInfo = SceneThumbnailInfoPtr.Get();
+	
+	if (!SceneThumbnailInfo)
 	{
 		if ( AssetThumbnail.IsValid() )
 		{
-			UObject* Asset = AssetThumbnail.Pin()->GetAsset();
-
-			if ( Asset )
+			if ( UObject* Asset = AssetThumbnail.Pin()->GetAsset() )
 			{
 				static const FName AssetToolsName("AssetTools");
 				FAssetToolsModule& AssetToolsModule = FModuleManager::LoadModuleChecked<FAssetToolsModule>(AssetToolsName);
@@ -335,30 +327,12 @@ USceneThumbnailInfo* SThumbnailEditModeTools::GetSceneThumbnailInfo()
 		}
 	}
 
-	return SceneThumbnailInfo.Get();
-}
-
-USceneThumbnailInfoWithPrimitive* SThumbnailEditModeTools::GetSceneThumbnailInfoWithPrimitive()
-{
-	USceneThumbnailInfo* ThumbnailInfo = SceneThumbnailInfo.Get();
-	if ( !ThumbnailInfo )
-	{
-		ThumbnailInfo = GetSceneThumbnailInfo();
-	}
-
-	return Cast<USceneThumbnailInfoWithPrimitive>( ThumbnailInfo );
+	return SceneThumbnailInfo;
 }
 
-USceneThumbnailInfoWithPrimitive* SThumbnailEditModeTools::ConstGetSceneThumbnailInfoWithPrimitive() const
+USceneThumbnailInfoWithPrimitive* SThumbnailEditModeTools::GetSceneThumbnailInfoWithPrimitive() const
 {
-	return Cast<USceneThumbnailInfoWithPrimitive>( SceneThumbnailInfo.Get() );
+	return Cast<USceneThumbnailInfoWithPrimitive>( GetSceneThumbnailInfo() );
 }
 
-void SThumbnailEditModeTools::OnAssetDataChanged()
-{
-	// Set the SceneThumbnailInfo pointer as needed
-	GetSceneThumbnailInfo();
-}
-
-
 #undef LOCTEXT_NAMESPACE
diff --git a/Engine/Source/Editor/ContentBrowser/Private/SThumbnailEditModeTools.h b/Engine/Source/Editor/ContentBrowser/Private/SThumbnailEditModeTools.h
--- a/Engine/Source/Editor/ContentBrowser/Private/SThumbnailEditModeTools.h
+++ b/Engine/Source/Editor/ContentBrowser/Private/SThumbnailEditModeTools.h
@@ -58,19 +58,18 @@ protected:
 	FReply ResetToDefault();
 
 	/** Helper accessors for ThumbnailInfo objects */
-	USceneThumbnailInfo* GetSceneThumbnailInfo();
-	USceneThumbnailInfoWithPrimitive* GetSceneThumbnailInfoWithPrimitive();
-	USceneThumbnailInfoWithPrimitive* ConstGetSceneThumbnailInfoWithPrimitive() const;
+	USceneThumbnailInfo* GetSceneThumbnailInfo() const;
+	USceneThumbnailInfoWithPrimitive* GetSceneThumbnailInfoWithPrimitive() const;
 
 	EThumbnailPrimType GetDefaultThumbnailType() const;
 
-	/** Event fired when the asset data for this asset is loaded or changed */
-	void OnAssetDataChanged();
-
 protected:
+	bool bInSmallView;
 	bool bModifiedThumbnailWhileDragging;
 	FIntPoint DragStartLocation;
 	TWeakPtr<FAssetThumbnail> AssetThumbnail;
-	TWeakObjectPtr<USceneThumbnailInfo> SceneThumbnailInfo;
-	bool bInSmallView;
+
+private:
+	// Never access this directly.
+	mutable TWeakObjectPtr<USceneThumbnailInfo> SceneThumbnailInfoPtr;
 };
diff --git a/Engine/Source/Editor/EditorWidgets/EditorWidgets.Build.cs b/Engine/Source/Editor/EditorWidgets/EditorWidgets.Build.cs
--- a/Engine/Source/Editor/EditorWidgets/EditorWidgets.Build.cs
+++ b/Engine/Source/Editor/EditorWidgets/EditorWidgets.Build.cs
@@ -19,7 +19,8 @@ public class EditorWidgets : ModuleRules
 				"UnrealEd",
 				"CoreUObject",
 				"ToolWidgets",
-				"EditorConfig"
+				"EditorConfig",
+				"AssetDefinition",
 			}
 		);
 		
diff --git a/Engine/Source/Editor/EditorWidgets/Private/Filters/CustomClassFilterData.cpp b/Engine/Source/Editor/EditorWidgets/Private/Filters/CustomClassFilterData.cpp
--- a/Engine/Source/Editor/EditorWidgets/Private/Filters/CustomClassFilterData.cpp
+++ b/Engine/Source/Editor/EditorWidgets/Private/Filters/CustomClassFilterData.cpp
@@ -6,25 +6,40 @@
 #include "IAssetTypeActions.h"
 #include "UObject/Class.h"
 #include "UObject/NameTypes.h"
+#include "Misc/AssetFilterData.h"
 
-void FCustomClassFilterData::AddCategory(TSharedPtr<FFilterCategory> InCategory)
+FCustomClassFilterData::FCustomClassFilterData(UAssetDefinition* InAssetDefinition, const FAssetFilterData& InAssetTypeAction)
+	: AssetDefinitionPtr(InAssetDefinition)
 {
-	Categories.AddUnique(InCategory);
+	Class = InAssetDefinition->GetAssetClass()->GetClass();
+	FilterName = InAssetTypeAction.Name;
+	FilterDisplayName = InAssetTypeAction.DisplayText;
+	ClassPathName = InAssetDefinition->GetAssetClass().ToSoftObjectPath().GetAssetPath();
+
+	Filter = InAssetTypeAction.Filter;
 }
 
-TSharedPtr<IAssetTypeActions> FCustomClassFilterData::GetAssetTypeActions() const
+FCustomClassFilterData::FCustomClassFilterData(UClass*InClass, TSharedPtr<FFilterCategory> InCategory, FLinearColor InColor)
+	: Class(InClass)
+	, Color(InColor)
 {
-	return AssetTypeActions.Pin();
+	Categories.Add(InCategory);
+	ClassPathName = Class->GetClassPathName();
+	FilterName = Class->GetFName().ToString();
+	FilterDisplayName = Class->GetDisplayNameText();
+
+	Filter.ClassPaths.Add(Class->GetClassPathName());
+	Filter.bRecursiveClasses = true;
 }
 
-UClass* FCustomClassFilterData::GetClass() const
+void FCustomClassFilterData::AddCategory(TSharedPtr<FFilterCategory> InCategory)
 {
-	if(TSharedPtr<IAssetTypeActions> AssetTypeActionsPin = AssetTypeActions.Pin())
-	{
-		return AssetTypeActionsPin->GetSupportedClass();
-	}
+	Categories.AddUnique(InCategory);
+}
 
-	return Class;
+UClass* FCustomClassFilterData::GetClass() const
+{
+	return Class.Get();
 }
 
 TArray<TSharedPtr<FFilterCategory>> FCustomClassFilterData::GetCategories() const
@@ -34,57 +49,31 @@ TArray<TSharedPtr<FFilterCategory>> FCustomClassFilterData::GetCategories() cons
 
 FLinearColor FCustomClassFilterData::GetColor() const
 {
-	if(TSharedPtr<IAssetTypeActions> AssetTypeActionsPin = AssetTypeActions.Pin())
+	// Kept this one because some assets can have their color rebound, is it worth keeping it dynamic here?  meh.
+	if (const UAssetDefinition* AssetDefinition = AssetDefinitionPtr.Get())
 	{
-		return AssetTypeActionsPin->GetTypeColor();
+		return AssetDefinition->GetAssetColor();
 	}
 
 	return Color;
 }
 
-void FCustomClassFilterData::BuildBackendFilter(FARFilter &Filter)
+void FCustomClassFilterData::BuildBackendFilter(FARFilter& OutFilter)
 {
-	if ( AssetTypeActions.IsValid() )
-	{
-		if (AssetTypeActions.Pin()->CanFilter())
-		{
-			AssetTypeActions.Pin()->BuildBackendFilter(Filter);
-		}
-	}
-	// If there is no AssetTypeAction for this filter, simply add the class name to the FARFilter
-	else
-	{
-		Filter.ClassPaths.Add(Class->GetClassPathName());
-		Filter.bRecursiveClasses = true;
-	}
+	OutFilter = Filter;
 }
 
 FText FCustomClassFilterData::GetName() const
 {
-	if(TSharedPtr<IAssetTypeActions> AssetTypeActionsPin = AssetTypeActions.Pin())
-	{
-		return AssetTypeActionsPin->GetName();
-	}
-
-	return Class->GetDisplayNameText();
+	return FilterDisplayName;
 }
 
 FString FCustomClassFilterData::GetFilterName() const
 {
-	if(TSharedPtr<IAssetTypeActions> AssetTypeActionsPin = AssetTypeActions.Pin())
-	{
-		return AssetTypeActionsPin->GetFilterName().ToString();
-	}
-
-	return Class->GetFName().ToString();
+	return FilterName;
 }
 
 FTopLevelAssetPath FCustomClassFilterData::GetClassPathName() const
 {
-	if(TSharedPtr<IAssetTypeActions> AssetTypeActionsPin = AssetTypeActions.Pin())
-	{
-		return AssetTypeActionsPin->GetClassPathName();
-	}
-
-	return Class->GetClassPathName();
+	return ClassPathName;
 }
diff --git a/Engine/Source/Editor/EditorWidgets/Public/Filters/CustomClassFilterData.h b/Engine/Source/Editor/EditorWidgets/Public/Filters/CustomClassFilterData.h
--- a/Engine/Source/Editor/EditorWidgets/Public/Filters/CustomClassFilterData.h
+++ b/Engine/Source/Editor/EditorWidgets/Public/Filters/CustomClassFilterData.h
@@ -12,6 +12,8 @@
 
 class FFilterCategory;
 class IAssetTypeActions;
+class UAssetDefinition;
+struct FAssetFilterData;
 class UClass;
 struct FARFilter;
 
@@ -25,29 +27,16 @@ class EDITORWIDGETS_API FCustomClassFilterData
 public:
 
 	/* You can provide an IAssetTypeActions for your type to get all the information from that */
-	FCustomClassFilterData(TWeakPtr<IAssetTypeActions> InAssetTypeAction)
-	: Class(nullptr)
-	, AssetTypeActions(InAssetTypeAction)
-	{
-		
-	}
+	FCustomClassFilterData(UAssetDefinition* InAssetDefinition, const FAssetFilterData& InAssetTypeAction);
 
 	/* Or you can provide a UClass, a Category and a Color to identify the filter */
-	FCustomClassFilterData(UClass *InClass, TSharedPtr<FFilterCategory> InCategory, FLinearColor InColor)
-	: Class(InClass)
-	, Color(InColor)
-	{
-		Categories.Add(InCategory);
-	}
+	FCustomClassFilterData(UClass* InClass, TSharedPtr<FFilterCategory> InCategory, FLinearColor InColor);
 
 	/* Unlike normal filters, Type Filters are allowed to belong to multiple categories */
 	void AddCategory(TSharedPtr<FFilterCategory> InCategory);
-
-	/* Get the AssetTypeActions associated with this filter if they exist */
-	TSharedPtr<IAssetTypeActions> GetAssetTypeActions() const;
-
+	
 	/* Get the UClass associated with this filter */
-	UClass *GetClass() const;
+	UClass* GetClass() const;
 
 	/* Get the Catgories this filter belongs to */
 	TArray<TSharedPtr<FFilterCategory>> GetCategories() const;
@@ -56,7 +45,7 @@ public:
 	FLinearColor GetColor() const;
 
 	/** Add this filter to the input BackendFilter */
-	void BuildBackendFilter(FARFilter &Filter);
+	void BuildBackendFilter(FARFilter& OutFilter);
 
 	/** Get the logical name of this filter class */
 	FText GetName() const;
@@ -72,7 +61,19 @@ public:
 private:
 
 	/** The Class the filter is associated with (if it does not have an AssetTypeAction) */
-	UClass *Class;
+	TWeakObjectPtr<UClass> Class;
+
+	/** The Filter name */
+	FString FilterName;
+	
+	/** The Filter display name */
+	FText FilterDisplayName;
+
+	/** The class path of the class associated with the filter */
+	FTopLevelAssetPath ClassPathName;
+
+	/** The actual asset registry filter. */
+	FARFilter Filter;
 
 	/** The Categories the filter will show up under */
 	TArray<TSharedPtr<FFilterCategory>> Categories;
@@ -80,6 +81,6 @@ private:
 	/** The Color of the filter (if it does not have an AssetTypeAction to get the color from) */
 	FLinearColor Color;
 
-	/** An optional AssetTypeAction to get information about the filter class */
-	TWeakPtr<IAssetTypeActions> AssetTypeActions;
+	/** An optional AssetDefinition to get information about the filter class */
+	TWeakObjectPtr<UAssetDefinition> AssetDefinitionPtr;
 };
\ No newline at end of file
diff --git a/Engine/Source/Editor/EditorWidgets/Public/Filters/SAssetFilterBar.h b/Engine/Source/Editor/EditorWidgets/Public/Filters/SAssetFilterBar.h
--- a/Engine/Source/Editor/EditorWidgets/Public/Filters/SAssetFilterBar.h
+++ b/Engine/Source/Editor/EditorWidgets/Public/Filters/SAssetFilterBar.h
@@ -2,6 +2,7 @@
 
 #pragma once
 
+#include "AssetDefinitionRegistry.h"
 #include "AssetRegistry/ARFilter.h"
 #include "AssetToolsModule.h"
 #include "AssetTypeCategories.h"
@@ -50,6 +51,7 @@
 #include "UObject/UObjectGlobals.h"
 #include "UObject/UnrealNames.h"
 #include "Widgets/DeclarativeSyntaxSupport.h"
+#include "Misc/AssetFilterData.h"
 
 #include "SAssetFilterBar.generated.h"
 
@@ -808,82 +810,65 @@ protected:
 
 		FAssetToolsModule& AssetToolsModule = FModuleManager::LoadModuleChecked<FAssetToolsModule>(TEXT("AssetTools"));
 
-		// Add the Basic category
-		TSharedPtr<FFilterCategory> BasicCategory = MakeShareable(new FFilterCategory(LOCTEXT("BasicFilter", "Common"), LOCTEXT("BasicFilterTooltip", "Filter by Common assets.")));
-		AssetFilterCategories.Add(EAssetTypeCategories::Basic, BasicCategory);
-
-		// We expand the basic category by default if we are creating it
-		CategoryToExpand = BasicCategory;
-		
-		// Add the advanced categories
-		TArray<FAdvancedAssetCategory> AdvancedAssetCategories;
-		AssetToolsModule.Get().GetAllAdvancedAssetCategories(/*out*/ AdvancedAssetCategories);
-
-		for (const FAdvancedAssetCategory& AdvancedAssetCategory : AdvancedAssetCategories)
 		{
-			const FText Tooltip = FText::Format(LOCTEXT("WildcardFilterTooltip", "Filter by {0} Assets."), AdvancedAssetCategory.CategoryName);
-			AssetFilterCategories.Add(AdvancedAssetCategory.CategoryType, MakeShareable(new FFilterCategory(AdvancedAssetCategory.CategoryName, Tooltip)));
+			// Add the Basic category
+			TSharedPtr<FFilterCategory> BasicCategory = MakeShareable(new FFilterCategory(LOCTEXT("BasicFilter", "Common"), LOCTEXT("BasicFilterTooltip", "Filter by Common assets.")));
+			AssetFilterCategories.Add(EAssetCategoryPaths::Basic.GetCategory(), BasicCategory);
+
+			// We expand the basic category by default if we are creating it
+			CategoryToExpand = BasicCategory;
 		}
 
-		TArray<TWeakPtr<IAssetTypeActions>> AssetTypeActionsList;
-		
-		// Get the browser type maps
-		AssetToolsModule.Get().GetAssetTypeActionsList(AssetTypeActionsList);
-		
-		// Sort the list
-		struct FCompareIAssetTypeActions
 		{
-			FORCEINLINE bool operator()( const TWeakPtr<IAssetTypeActions>& A, const TWeakPtr<IAssetTypeActions>& B ) const
-			{
-				return A.Pin()->GetName().CompareTo( B.Pin()->GetName() ) == -1;
-			}
-		};
-		AssetTypeActionsList.Sort( FCompareIAssetTypeActions() );
+			TArray<TObjectPtr<UAssetDefinition>> AssetDefinitions = UAssetDefinitionRegistry::Get()->GetAllAssetDefinitions();
 
-		const TSharedRef<FPathPermissionList>& AssetClassPermissionList = AssetToolsModule.Get().GetAssetClassPathPermissionList(AssetClassAction);
+			const TSharedRef<FPathPermissionList>& AssetClassPermissionList = AssetToolsModule.Get().GetAssetClassPathPermissionList(AssetClassAction);
 
-		// For every asset type, convert it to an FCustomClassFilterData and add it to the list
-		for (int32 ClassIdx = 0; ClassIdx < AssetTypeActionsList.Num(); ++ClassIdx)
-		{
-			const TWeakPtr<IAssetTypeActions>& WeakTypeActions = AssetTypeActionsList[ClassIdx];
-			if ( WeakTypeActions.IsValid() )
+			// For every asset type, convert it to an FCustomClassFilterData and add it to the list
+			for (UAssetDefinition* AssetDefinition : AssetDefinitions)
 			{
-				TSharedPtr<IAssetTypeActions> TypeActions = WeakTypeActions.Pin();
-				if ( ensure(TypeActions.IsValid()) && TypeActions->CanFilter() )
+				TSoftClassPtr<UObject> AssetClass = AssetDefinition->GetAssetClass();
+				if ((!AssetClass.IsNull() && AssetClassPermissionList->PassesFilter(AssetClass.ToString())))
 				{
-					UClass* SupportedClass = TypeActions->GetSupportedClass();
-					if ((!SupportedClass || AssetClassPermissionList->PassesFilter(SupportedClass->GetPathName())))
+					const TSharedRef<FAssetFilterDataCache> FilterCache = AssetDefinition->GetFilters();
+				
+					for (const FAssetFilterData& FilterData : FilterCache->Filters)
 					{
 						// Convert the AssetTypeAction to an FCustomClassFilterData and add it to our list
-						TSharedRef<FCustomClassFilterData> CustomClassFilterData = MakeShared<FCustomClassFilterData>(TypeActions);
+						TSharedRef<FCustomClassFilterData> CustomClassFilterData = MakeShared<FCustomClassFilterData>(AssetDefinition, FilterData);
+
+						for (const FAssetCategoryPath& CategoryPath : FilterData.FilterCategories)
+						{
+							TSharedPtr<FFilterCategory> FilterCategory = AssetFilterCategories.FindRef(CategoryPath.GetCategory());
+							if (!FilterCategory.IsValid())
+							{
+								const FText Tooltip = FText::Format(LOCTEXT("WildcardFilterTooltip", "Filter by {0} Assets."), CategoryPath.GetCategoryText());
+								FilterCategory = MakeShared<FFilterCategory>(CategoryPath.GetCategoryText(), Tooltip);
+								AssetFilterCategories.Add(CategoryPath.GetCategory(), FilterCategory);
+							}
+
+							CustomClassFilterData->AddCategory(FilterCategory);
+						}
+						
 						CustomClassFilters.Add(CustomClassFilterData);
 					}
 				}
 			}
 		}
-		
-		// Do a second pass through all the CustomClassFilters with AssetTypeActions to update their categories
-		UpdateAssetTypeActionCategories();
 
 		// Update/remove any asset filters that already exist and use AssetTypeActions
 		TArray<TSharedRef<SAssetFilter>> AssetFiltersToRemove;
 		
-		for(TSharedRef<SAssetFilter>& AssetFilter : AssetFilters)
+		for (TSharedRef<SAssetFilter>& AssetFilter : AssetFilters)
 		{
-			// We only care about filters that are based on asset type actions
-			if(!AssetFilter->GetCustomClassFilterData()->GetAssetTypeActions())
-			{
-				continue;
-			}
-
 			// Try and find the FCustomClassFilterData for the current filter using the new list we just created in CustomClassFilters
 			TSharedRef<FCustomClassFilterData>* FoundCustomClassFilterData = CustomClassFilters.FindByPredicate([AssetFilter](const TSharedRef<FCustomClassFilterData>& CustomClassFilterData)
 			{
-				return CustomClassFilterData->GetAssetTypeActions() == AssetFilter->GetCustomClassFilterData()->GetAssetTypeActions();
+				return CustomClassFilterData->GetFilterName() == AssetFilter->GetCustomClassFilterData()->GetFilterName();
 			});
 
 			// If it exists, update the filter to let it know
-			if(FoundCustomClassFilterData)
+			if (FoundCustomClassFilterData)
 			{
 				AssetFilter->SetCustomClassFilterData(*FoundCustomClassFilterData);
 			}
@@ -894,30 +879,11 @@ protected:
 			}
 		}
 
-		for(TSharedRef<SAssetFilter>& AssetFilter : AssetFiltersToRemove)
+		for (TSharedRef<SAssetFilter>& AssetFilter : AssetFiltersToRemove)
 		{
 			RemoveFilter(AssetFilter);
 		}
 	}
-
-	void UpdateAssetTypeActionCategories()
-	{
-		for(const TSharedRef<FCustomClassFilterData> &CustomClassFilter : CustomClassFilters)
-		{
-			for (auto MenuIt = AssetFilterCategories.CreateIterator(); MenuIt; ++MenuIt)
-			{
-				if(TSharedPtr<IAssetTypeActions> AssetTypeActions = CustomClassFilter->GetAssetTypeActions())
-				{
-					if(MenuIt.Key() & AssetTypeActions->GetCategories())
-					{
-						CustomClassFilter->AddCategory(MenuIt.Value());
-					}
-				}
-				
-			}
-		}	
-	}
-
 	
 	/* Filter Dropdown Related Functionality */
 
@@ -1054,7 +1020,7 @@ protected:
 		}
 
 		// Set the extension hook for the basic category, if it exists and we have any assets for it
-		if(TSharedPtr<FFilterCategory>* BasicCategory = AssetFilterCategories.Find(EAssetTypeCategories::Basic))
+		if(TSharedPtr<FFilterCategory>* BasicCategory = AssetFilterCategories.Find(EAssetCategoryPaths::Basic.GetCategory()))
 		{
 			if(FCategoryMenu* BasicMenu = CategoryToMenuMap.Find(*BasicCategory))
 			{
@@ -1245,13 +1211,13 @@ protected:
 	/** Copy of Custom Class Filters that were provided by the user, and not autopopulated from AssetTypeActions */
 	TArray<TSharedRef<FCustomClassFilterData>> UserCustomClassFilters;
 
-	TMap<EAssetTypeCategories::Type, TSharedPtr<FFilterCategory>> AssetFilterCategories;
+	TMap<FName, TSharedPtr<FFilterCategory>> AssetFilterCategories;
 
 	/** The category to expand in the filter menu */
 	TSharedPtr<FFilterCategory> CategoryToExpand;
 	
 	/** Whether the filter bar provides the default Asset Filters */
-	bool bUseDefaultAssetFilters;
+	bool bUseDefaultAssetFilters = false;
 
 	/** The AssetClassAction used to get the permission list from the Asset Tools Module */
 	EAssetClassAction AssetClassAction = EAssetClassAction::ViewAsset;
diff --git a/Engine/Source/Editor/SceneOutliner/SceneOutliner.Build.cs b/Engine/Source/Editor/SceneOutliner/SceneOutliner.Build.cs
--- a/Engine/Source/Editor/SceneOutliner/SceneOutliner.Build.cs
+++ b/Engine/Source/Editor/SceneOutliner/SceneOutliner.Build.cs
@@ -26,7 +26,7 @@ public class SceneOutliner : ModuleRules
 				"EditorWidgets",
 				"ToolWidgets",
 				"UnsavedAssetsTracker",
-				"UncontrolledChangelists"
+				"AssetDefinition"
 			}
 		);
 		
diff --git a/Engine/Source/Editor/UnrealEd/Private/AssetDefinitionDefault.cpp b/Engine/Source/Editor/UnrealEd/Private/AssetDefinitionDefault.cpp
--- a/Engine/Source/Editor/UnrealEd/Private/AssetDefinitionDefault.cpp
+++ b/Engine/Source/Editor/UnrealEd/Private/AssetDefinitionDefault.cpp
@@ -4,6 +4,7 @@
 
 #include "AssetToolsModule.h"
 #include "EditorFramework/AssetImportData.h"
+#include "EditorFramework/ThumbnailInfo.h"
 #include "Settings/EditorLoadingSavingSettings.h"
 #include "Toolkits/SimpleAssetEditor.h"
 
@@ -37,4 +38,35 @@ EAssetCommandResult UAssetDefinitionDefault::PerformAssetDiff(const FAssetDiffAr
 	return EAssetCommandResult::Handled;
 }
 
+namespace UE::Editor
+{
+	UThumbnailInfo* FindOrCreateThumbnailInfo(UObject* AssetObject, TSubclassOf<UThumbnailInfo> ThumbnailClass)
+	{
+		if (AssetObject && ThumbnailClass)
+		{
+			if (const FObjectProperty* ObjectProperty = FindFProperty<FObjectProperty>(AssetObject->GetClass(), "ThumbnailInfo"))
+			{
+				// Is the property marked as Instanced?
+				if (ObjectProperty->HasAllPropertyFlags(CPF_PersistentInstance | CPF_ExportObject | CPF_InstancedReference))
+				{
+					// Get the thumbnail.
+					UThumbnailInfo* ThumbnailInfo = Cast<UThumbnailInfo>(ObjectProperty->GetObjectPropertyValue_InContainer(AssetObject));
+					if (ThumbnailInfo && ThumbnailInfo->GetClass() == ThumbnailClass)
+					{
+						return ThumbnailInfo;
+					}
+
+					// We couldn't find it, need to initialize it.
+					ThumbnailInfo = NewObject<UThumbnailInfo>(AssetObject, ThumbnailClass, NAME_None, RF_Transactional);
+					ObjectProperty->SetObjectPropertyValue_InContainer(AssetObject, ThumbnailInfo);
+
+					return ThumbnailInfo;
+				}
+			}
+		}
+		
+		return nullptr;
+	}
+}
+
 #undef LOCTEXT_NAMESPACE
diff --git a/Engine/Source/Editor/UnrealEd/Private/Kismet2/BlueprintEditorUtils.cpp b/Engine/Source/Editor/UnrealEd/Private/Kismet2/BlueprintEditorUtils.cpp
--- a/Engine/Source/Editor/UnrealEd/Private/Kismet2/BlueprintEditorUtils.cpp
+++ b/Engine/Source/Editor/UnrealEd/Private/Kismet2/BlueprintEditorUtils.cpp
@@ -10608,6 +10608,15 @@ bool FBlueprintEditorUtils::HasFunctionBlueprintThreadSafeMetaData(const UFuncti
 	}
 	
 	return false;
-};
+}
+
+bool FBlueprintEditorUtils::ShouldOpenWithDataOnlyEditor(const UBlueprint* Blueprint)
+{
+	return FBlueprintEditorUtils::IsDataOnlyBlueprint(Blueprint)
+		&& !FBlueprintEditorUtils::IsLevelScriptBlueprint(Blueprint)
+		&& !FBlueprintEditorUtils::IsInterfaceBlueprint(Blueprint)
+		&& !Blueprint->bForceFullEditor
+		&& !Blueprint->bIsNewlyCreated;
+}
 
 #undef LOCTEXT_NAMESPACE
diff --git a/Engine/Source/Editor/UnrealEd/Public/AssetDefinitionDefault.h b/Engine/Source/Editor/UnrealEd/Public/AssetDefinitionDefault.h
--- a/Engine/Source/Editor/UnrealEd/Public/AssetDefinitionDefault.h
+++ b/Engine/Source/Editor/UnrealEd/Public/AssetDefinitionDefault.h
@@ -18,3 +18,25 @@ public:
 	virtual EAssetCommandResult PerformAssetDiff(const FAssetDiffArgs& DiffArgs) const override;
 	// UAssetDefinition End
 };
+
+namespace UE::Editor
+{
+	/**
+	 * This function will find or set a new thumbnail info on the object if it either does not exist or the types don't
+	 * match.  This function has some requirements, such as the property MUST minimally exist as such,
+	 * 
+	 * UPROPERTY(Instanced)
+	 * TObjectPtr<class UThumbnailInfo> ThumbnailInfo;
+	 * 
+	 * If your class needs the thumbnail to be set using a Setter/Getter, that's ok, you can have 
+	 * Getter=GetThumbnailInfo, Setter=SetThumbnailInfo, in the UProperty declaration to take advantage of native
+	 * setters and getters.
+	 */
+	UNREALED_API UThumbnailInfo* FindOrCreateThumbnailInfo(UObject* AssetObject, TSubclassOf<UThumbnailInfo> ThumbnailClass);
+
+	template<typename ThumbnailClass>
+	ThumbnailClass* FindOrCreateThumbnailInfo(UObject* AssetObject)
+	{
+		return Cast<ThumbnailClass>(FindOrCreateThumbnailInfo(AssetObject, ThumbnailClass::StaticClass()));
+	}
+}
diff --git a/Engine/Source/Editor/UnrealEd/Public/Kismet2/BlueprintEditorUtils.h b/Engine/Source/Editor/UnrealEd/Public/Kismet2/BlueprintEditorUtils.h
--- a/Engine/Source/Editor/UnrealEd/Public/Kismet2/BlueprintEditorUtils.h
+++ b/Engine/Source/Editor/UnrealEd/Public/Kismet2/BlueprintEditorUtils.h
@@ -1913,6 +1913,9 @@ public:
 	 * been validated for thread safe use in a blueprint.
 	 */
 	static bool HasFunctionBlueprintThreadSafeMetaData(const UFunction* InFunction);
+
+public:
+	static bool ShouldOpenWithDataOnlyEditor(const UBlueprint* Blueprint);
 };
 
 struct UNREALED_API FBlueprintDuplicationScopeFlags
diff --git a/Engine/Source/Runtime/Engine/Classes/Engine/SkeletalMesh.h b/Engine/Source/Runtime/Engine/Classes/Engine/SkeletalMesh.h
--- a/Engine/Source/Runtime/Engine/Classes/Engine/SkeletalMesh.h
+++ b/Engine/Source/Runtime/Engine/Classes/Engine/SkeletalMesh.h
@@ -61,6 +61,7 @@ class FSkinnedAssetBuildContext;
 class FSkinnedAssetPostLoadContext;
 class USkeletalMeshLODSettings;
 class USkeleton;
+class UThumbnailInfo;
 struct FMeshUVChannelInfo;
 struct FSkeletalMaterial;
 struct FSkinWeightProfileInfo;
@@ -1521,8 +1522,8 @@ public:
 
 	/** Information for thumbnail rendering */
 	UE_DEPRECATED(4.27, "Please do not access this member directly; use USkeletalMesh::GetThumbnailInfo() or USkeletalMesh::SetThumbnailInfo().")
-	UPROPERTY(VisibleAnywhere, Instanced, AdvancedDisplay, Category = Thumbnail)
-	TObjectPtr<class UThumbnailInfo> ThumbnailInfo;
+	UPROPERTY(VisibleAnywhere, Instanced, AdvancedDisplay, Getter=GetThumbnailInfo, Setter=SetThumbnailInfo, Category = Thumbnail)
+	TObjectPtr<UThumbnailInfo> ThumbnailInfo;
 	
 	static FName GetThumbnailInfoMemberName()
 	{
@@ -1530,16 +1531,16 @@ public:
 		return GET_MEMBER_NAME_CHECKED(USkeletalMesh, ThumbnailInfo);
 		PRAGMA_ENABLE_DEPRECATION_WARNINGS
 	}
-
-	class UThumbnailInfo* GetThumbnailInfo() const
+	
+	UThumbnailInfo* GetThumbnailInfo() const
 	{
 		WaitUntilAsyncPropertyReleased(ESkeletalMeshAsyncProperties::ThumbnailInfo);
 		PRAGMA_DISABLE_DEPRECATION_WARNINGS
 		return ThumbnailInfo;
 		PRAGMA_ENABLE_DEPRECATION_WARNINGS
 	}
-
-	void SetThumbnailInfo(class UThumbnailInfo* InThumbnailInfo)
+	
+	void SetThumbnailInfo(UThumbnailInfo* InThumbnailInfo)
 	{
 		WaitUntilAsyncPropertyReleased(ESkeletalMeshAsyncProperties::ThumbnailInfo);
 		PRAGMA_DISABLE_DEPRECATION_WARNINGS
