AssetDefinition - Fixing filter categories to include everything that hasn't yet been converted to a Asset Definition.

AssetDefintiion - Converted more blueprint / Animation blueprint stuff over to be an asset definition.

diff --git a/Engine/Plugins/Editor/EngineAssetDefinitions/Source/EngineAssetDefinitions.Build.cs b/Engine/Plugins/Editor/EngineAssetDefinitions/Source/EngineAssetDefinitions.Build.cs
--- a/Engine/Plugins/Editor/EngineAssetDefinitions/Source/EngineAssetDefinitions.Build.cs
+++ b/Engine/Plugins/Editor/EngineAssetDefinitions/Source/EngineAssetDefinitions.Build.cs
@@ -34,7 +34,8 @@ public class EngineAssetDefinitions : ModuleRules
 				"Landscape",
 				"PhysicsCore",
 				"Kismet",
-				"GameProjectGeneration"
+				"GameProjectGeneration",
+				"AnimationBlueprintEditor"
 			}
 		);
 		
diff --git a/Engine/Plugins/Editor/EngineAssetDefinitions/Source/Private/Animation/AssetDefinition_AnimBlueprint.cpp b/Engine/Plugins/Editor/EngineAssetDefinitions/Source/Private/Animation/AssetDefinition_AnimBlueprint.cpp
new file mode 100644
--- /dev/null
+++ b/Engine/Plugins/Editor/EngineAssetDefinitions/Source/Private/Animation/AssetDefinition_AnimBlueprint.cpp
@@ -0,0 +1,393 @@
+// Copyright Epic Games, Inc. All Rights Reserved.
+
+#include "Animation/AssetDefinition_AnimBlueprint.h"
+
+#include "AnimationEditorUtils.h"
+#include "ContentBrowserMenuContexts.h"
+#include "ToolMenus.h"
+#include "Widgets/Layout/SBorder.h"
+#include "Misc/MessageDialog.h"
+#include "Widgets/Images/SImage.h"
+#include "Styling/AppStyle.h"
+#include "Animation/AnimInstance.h"
+#include "Factories/AnimBlueprintFactory.h"
+#include "ContentBrowserModule.h"
+#include "Framework/Application/SlateApplication.h"
+#include "SBlueprintDiff.h"
+#include "SSkeletonWidget.h"
+#include "Styling/SlateIconFinder.h"
+#include "IContentBrowserSingleton.h"
+#include "Kismet2/KismetEditorUtilities.h"
+#include "Subsystems/AssetEditorSubsystem.h"
+#include "Editor.h"
+#include "IAnimationBlueprintEditorModule.h"
+#include "IAssetTools.h"
+#include "Algo/AllOf.h"
+#include "Algo/NoneOf.h"
+
+#define LOCTEXT_NAMESPACE "AssetTypeActions"
+
+namespace UE::AnimBlueprint
+{
+	bool ReplaceMissingSkeleton(const TArray<TObjectPtr<UObject>>& InAnimationAssets)
+	{
+		// record anim assets that need skeleton replaced
+		const TArray<TWeakObjectPtr<UObject>> AnimsToFix = TArray<TWeakObjectPtr<UObject>>(InAnimationAssets);
+		
+		// get a skeleton from the user and replace it
+		const TSharedPtr<SReplaceMissingSkeletonDialog> PickSkeletonWindow = SNew(SReplaceMissingSkeletonDialog).AnimAssets(AnimsToFix);
+		const bool bWasSkeletonReplaced = PickSkeletonWindow.Get()->ShowModal();
+		return bWasSkeletonReplaced;
+	}
+
+	EBlueprintType GetBlueprintType(const FAssetData& InAssetData)
+	{
+		const UEnum* BlueprintTypeEnum = StaticEnum<EBlueprintType>();
+		const FString EnumString = InAssetData.GetTagValueRef<FString>(GET_MEMBER_NAME_CHECKED(UBlueprint, BlueprintType));
+		const EBlueprintType BlueprintType = (!EnumString.IsEmpty() ? static_cast<EBlueprintType>(BlueprintTypeEnum->GetValueByName(*EnumString)) : BPTYPE_Normal);
+		return BlueprintType;
+	}
+
+	bool IsAnimBlueprintTemplate(const FAssetData& InAssetData)
+	{
+		return InAssetData.GetTagValueRef<bool>(GET_MEMBER_NAME_CHECKED(UAnimBlueprint, bIsTemplate));
+	}
+
+	TSoftObjectPtr<USkeleton> GetAnimBlueprintTargetSkeleton(const FAssetData& InAssetData)
+	{
+		return TSoftObjectPtr<USkeleton>(FSoftObjectPath(InAssetData.GetTagValueRef<FString>(GET_MEMBER_NAME_CHECKED(UAnimBlueprint, TargetSkeleton))));
+	}
+}
+
+FText UAssetDefinition_AnimBlueprint::GetAssetDisplayName(const FAssetData& AssetData) const
+{
+	FString OutBlueprintType;
+	if (AssetData.GetTagValue(GET_MEMBER_NAME_CHECKED(UBlueprint, BlueprintType), OutBlueprintType))
+	{
+		if (OutBlueprintType == TEXT("BPTYPE_Interface"))
+		{
+			return LOCTEXT("AssetTypeActions_AnimLayerInterface", "Animation Layer Interface");
+		}
+	}
+
+	return Super::GetAssetDisplayName(AssetData);
+}
+
+void UAssetDefinition_AnimBlueprint::BuildFilters(TArray<FAssetFilterData>& OutFilters) const
+{
+	// Intentionally skipping UAssetDefinition_Blueprint.
+	UAssetDefinition::BuildFilters(OutFilters);
+
+	{
+		FAssetFilterData Filter;
+		Filter.Name = GetAssetClass().ToSoftObjectPath().ToString() + TEXT("Interface");
+		Filter.DisplayText = LOCTEXT("AssetTypeActions_AnimLayerInterface", "Animation Layer Interface");
+		Filter.FilterCategories = { EAssetCategoryPaths::Animation };
+		Filter.Filter.ClassPaths.Add(GetAssetClass().ToSoftObjectPath().GetAssetPath());
+		Filter.Filter.TagsAndValues.Add(GET_MEMBER_NAME_CHECKED(UBlueprint, BlueprintType), FString(TEXT("BPTYPE_Interface")));
+		OutFilters.Add(MoveTemp(Filter));
+	}
+}
+
+UFactory* UAssetDefinition_AnimBlueprint::GetFactoryForBlueprintType(UBlueprint* InBlueprint) const
+{
+	UAnimBlueprint* AnimBlueprint = CastChecked<UAnimBlueprint>(InBlueprint);
+
+	if (InBlueprint->BlueprintType == BPTYPE_Interface)
+	{
+		return NewObject<UAnimLayerInterfaceFactory>();
+	}
+	else
+	{
+		UAnimBlueprintFactory* AnimBlueprintFactory = NewObject<UAnimBlueprintFactory>();
+		AnimBlueprintFactory->ParentClass = TSubclassOf<UAnimInstance>(*InBlueprint->GeneratedClass);
+		AnimBlueprintFactory->TargetSkeleton = AnimBlueprint->TargetSkeleton;
+		AnimBlueprintFactory->bTemplate = AnimBlueprint->bIsTemplate;
+		return AnimBlueprintFactory;
+	}
+}
+
+TSharedPtr<SWidget> UAssetDefinition_AnimBlueprint::GetThumbnailOverlay(const FAssetData& AssetData) const
+{
+	const FSlateBrush* Icon = FSlateIconFinder::FindIconBrushForClass(UAnimBlueprint::StaticClass());
+
+	return SNew(SBorder)
+		.BorderImage(FAppStyle::GetNoBrush())
+		.Visibility(EVisibility::HitTestInvisible)
+		.Padding(FMargin(0.0f, 0.0f, 0.0f, 3.0f))
+		.HAlign(HAlign_Right)
+		.VAlign(VAlign_Bottom)
+		[
+			SNew(SImage)
+			.Image(Icon)
+		];
+}
+
+EAssetCommandResult UAssetDefinition_AnimBlueprint::PerformAssetDiff(const FAssetDiffArgs& DiffArgs) const
+{
+	UBlueprint* OldBlueprint = CastChecked<UBlueprint>(DiffArgs.OldAsset);
+	UBlueprint* NewBlueprint = CastChecked<UBlueprint>(DiffArgs.NewAsset);
+
+	// sometimes we're comparing different revisions of one single asset (other 
+	// times we're comparing two completely separate assets altogether)
+	bool bIsSingleAsset = (NewBlueprint->GetName() == OldBlueprint->GetName());
+
+	FText WindowTitle = LOCTEXT("NamelessAnimationBlueprintDiff", "Animation Blueprint Diff");
+	// if we're diffing one asset against itself 
+	if (bIsSingleAsset)
+	{
+		// identify the assumed single asset in the window's title
+		WindowTitle = FText::Format(LOCTEXT("AnimationBlueprintDiff", "{0} - Animation Blueprint Diff"), FText::FromString(NewBlueprint->GetName()));
+	}
+
+	SBlueprintDiff::CreateDiffWindow(WindowTitle, OldBlueprint, NewBlueprint, DiffArgs.OldRevision, DiffArgs.NewRevision);
+
+	return EAssetCommandResult::Handled;
+}
+
+
+EAssetCommandResult UAssetDefinition_AnimBlueprint::OpenAssets(const FAssetOpenArgs& OpenArgs) const
+{
+	TArray<FAssetData> OutAssetsThatFailedToLoad;
+	for (UAnimBlueprint* AnimBlueprint : OpenArgs.LoadObjects<UAnimBlueprint>({}, &OutAssetsThatFailedToLoad))
+	{
+		if (AnimBlueprint->SkeletonGeneratedClass && AnimBlueprint->GeneratedClass)
+		{
+			if (AnimBlueprint->BlueprintType != BPTYPE_Interface && !AnimBlueprint->TargetSkeleton && !AnimBlueprint->bIsTemplate)
+			{
+				FText ShouldRetargetMessage = LOCTEXT("ShouldRetarget_Message", "Could not find the skeleton for Anim Blueprint '{BlueprintName}' Would you like to choose a new one?");
+				
+				FFormatNamedArguments Arguments;
+				Arguments.Add( TEXT("BlueprintName"), FText::FromString(AnimBlueprint->GetName()));
+
+				if (FMessageDialog::Open(EAppMsgType::YesNo, FText::Format(ShouldRetargetMessage, Arguments)) == EAppReturnType::Yes)
+				{
+					TArray<TObjectPtr<UObject>> AssetsToRetarget;
+					AssetsToRetarget.Add(AnimBlueprint);
+					const bool bSkeletonReplaced = UE::AnimBlueprint::ReplaceMissingSkeleton(AssetsToRetarget);
+					if (!bSkeletonReplaced)
+					{
+						return EAssetCommandResult::Handled; // Persona will crash if trying to load asset without a skeleton
+					}
+				}
+				else
+				{
+					return EAssetCommandResult::Handled;
+				}
+			}
+			
+			const bool bBringToFrontIfOpen = true;
+#if WITH_EDITOR
+			if (IAssetEditorInstance* EditorInstance = GEditor->GetEditorSubsystem<UAssetEditorSubsystem>()->FindEditorForAsset(AnimBlueprint, bBringToFrontIfOpen))
+			{
+				EditorInstance->FocusWindow(AnimBlueprint);
+			}
+			else
+#endif
+			{
+				IAnimationBlueprintEditorModule& AnimationBlueprintEditorModule = FModuleManager::LoadModuleChecked<IAnimationBlueprintEditorModule>("AnimationBlueprintEditor");
+				AnimationBlueprintEditorModule.CreateAnimationBlueprintEditor(OpenArgs.GetToolkitMode(), OpenArgs.ToolkitHost, AnimBlueprint);
+			}
+		}
+		else
+		{
+			FMessageDialog::Open( EAppMsgType::Ok, LOCTEXT("FailedToLoadCorruptAnimBlueprint", "The Anim Blueprint could not be loaded because it is corrupt."));
+		}
+	}
+
+	return EAssetCommandResult::Handled;
+}
+
+// Menu Extensions
+//--------------------------------------------------------------------
+
+namespace MenuExtension_AnimBlueprint
+{
+	void ExecuteFindSkeleton(const FToolMenuContext& InContext)
+	{
+		const UContentBrowserAssetContextMenuContext* CBContext = UContentBrowserAssetContextMenuContext::FindContextWithAssets(InContext);
+
+		TArray<UObject*> ObjectsToSync;
+		for (UAnimBlueprint* AnimBlueprint : CBContext->LoadSelectedObjects<UAnimBlueprint>())
+		{
+			if (USkeleton* Skeleton = AnimBlueprint->TargetSkeleton)
+			{
+				ObjectsToSync.AddUnique(Skeleton);
+			}
+		}
+    
+		if ( ObjectsToSync.Num() > 0 )
+		{
+			IAssetTools::Get().SyncBrowserToAssets(ObjectsToSync);
+		}
+	}
+
+	void ExecuteAssignSkeleton(const FToolMenuContext& InContext)
+	{
+		const UContentBrowserAssetContextMenuContext* CBContext = UContentBrowserAssetContextMenuContext::FindContextWithAssets(InContext);
+
+		TSharedRef<SWindow> WidgetWindow = SNew(SWindow)
+			.Title(LOCTEXT("ChooseSkeletonWindowTitle", "Choose Skeleton"))
+			.ClientSize(FVector2D(400, 600));
+
+		TSharedPtr<SSkeletonSelectorWindow> SkeletonSelectorWindow;
+		WidgetWindow->SetContent
+		(
+			SNew(SBorder)
+			.BorderImage(FAppStyle::GetBrush("ToolPanel.GroupBorder"))
+			[
+				SAssignNew(SkeletonSelectorWindow, SSkeletonSelectorWindow)
+				.WidgetWindow(WidgetWindow)
+			]
+		);
+
+		GEditor->EditorAddModalWindow(WidgetWindow);
+
+		if (USkeleton* SelectedSkeleton = SkeletonSelectorWindow->GetSelectedSkeleton())
+		{
+			for (UAnimBlueprint* AnimBlueprint : CBContext->LoadSelectedObjects<UAnimBlueprint>())
+			{
+				if (AnimBlueprint->TargetSkeleton != SelectedSkeleton)
+				{
+					AnimBlueprint->Modify();
+					AnimBlueprint->TargetSkeleton = SelectedSkeleton;
+				}
+			}
+		}
+	}
+
+	void BuildNewSkeletonChildBlueprintMenu(UToolMenu* Menu, const FAssetData InAnimBlueprintAsset)
+	{
+		auto HandleAssetSelected = [InAnimBlueprintAsset](const FAssetData& InSelectedSkeletonAsset)
+		{
+			FSlateApplication::Get().DismissAllMenus();
+			
+			if (UAnimBlueprint* TargetParentBP = Cast<UAnimBlueprint>(InAnimBlueprintAsset.GetAsset()))
+			{
+				USkeleton* TargetSkeleton = CastChecked<USkeleton>(InSelectedSkeletonAsset.GetAsset());
+				UClass* TargetParentClass = TargetParentBP->GeneratedClass;
+
+				if (!FKismetEditorUtilities::CanCreateBlueprintOfClass(TargetParentClass))
+				{
+					FMessageDialog::Open( EAppMsgType::Ok, LOCTEXT("InvalidClassToMakeBlueprintFrom", "Invalid class with which to make a Blueprint."));
+					return;
+				}
+
+				FString Name;
+				FString PackageName;
+				IAssetTools::Get().CreateUniqueAssetName(TargetParentBP->GetOutermost()->GetName(), TEXT("_Child"), PackageName, Name);
+				const FString PackagePath = FPackageName::GetLongPackagePath(PackageName);
+
+				UAnimBlueprintFactory* AnimBlueprintFactory = NewObject<UAnimBlueprintFactory>();
+				AnimBlueprintFactory->ParentClass = TSubclassOf<UAnimInstance>(*TargetParentBP->GeneratedClass);
+				AnimBlueprintFactory->TargetSkeleton = TargetSkeleton;
+				AnimBlueprintFactory->bTemplate = false;
+
+				FContentBrowserModule& ContentBrowserModule = FModuleManager::LoadModuleChecked<FContentBrowserModule>("ContentBrowser");
+				ContentBrowserModule.Get().CreateNewAsset(Name, PackagePath, TargetParentBP->GetClass(), AnimBlueprintFactory);
+			}
+		};
+
+		TSoftObjectPtr<USkeleton> AnimBlueprintSkeletonPtr = UE::AnimBlueprint::GetAnimBlueprintTargetSkeleton(InAnimBlueprintAsset);
+		if (USkeleton* AnimBlueprintSkeleton = AnimBlueprintSkeletonPtr.Get())
+		{
+			TArray<FAssetData> CompatibleSkeletonAssets;
+			AnimBlueprintSkeleton->GetCompatibleSkeletonAssets(CompatibleSkeletonAssets);
+
+			TArray<FSoftObjectPath> CompatibleSkeletonPaths;
+			Algo::Transform(CompatibleSkeletonAssets, CompatibleSkeletonPaths, [](const FAssetData& InAsset) { return InAsset.GetSoftObjectPath(); });
+		
+
+			FAssetPickerConfig AssetPickerConfig;
+			AssetPickerConfig.OnAssetEnterPressed = FOnAssetEnterPressed::CreateLambda([HandleAssetSelected](const TArray<FAssetData>& SelectedAssetData)
+			{
+				if (SelectedAssetData.Num() == 1)
+				{
+					HandleAssetSelected(SelectedAssetData[0]);
+				}
+			});
+			AssetPickerConfig.OnAssetSelected = FOnAssetSelected::CreateLambda(HandleAssetSelected);
+			AssetPickerConfig.bAllowNullSelection = false;
+			AssetPickerConfig.InitialAssetViewType = EAssetViewType::List;
+			AssetPickerConfig.Filter.SoftObjectPaths = MoveTemp(CompatibleSkeletonPaths);
+		
+			FContentBrowserModule& ContentBrowserModule = FModuleManager::Get().LoadModuleChecked<FContentBrowserModule>(TEXT("ContentBrowser"));
+
+			FToolMenuSection& InSection = Menu->AddSection("CompatibleSkeletonMenu", LOCTEXT("CompatibleSkeletonHeader", "Compatible Skeletons"));
+			InSection.AddEntry(
+				FToolMenuEntry::InitWidget("CompatibleSkeletonPicker",
+					SNew(SBox)
+					.WidthOverride(300.0f)
+					.HeightOverride(300.0f)
+					[
+						ContentBrowserModule.Get().CreateAssetPicker(AssetPickerConfig)
+					],
+					FText::GetEmpty()
+				)
+			);
+		}
+	}
+		
+	static FDelayedAutoRegisterHelper DelayedAutoRegister(EDelayedRegisterRunPhase::EndOfEngineInit, []{ 
+		UToolMenus::RegisterStartupCallback(FSimpleMulticastDelegate::FDelegate::CreateLambda([]()
+		{
+			FToolMenuOwnerScoped OwnerScoped(UE_MODULE_NAME);
+			UToolMenu* Menu = UE::ContentBrowser::ExtendToolMenu_AssetContextMenu(UBlueprint::StaticClass());
+	        
+			FToolMenuSection& Section = Menu->FindOrAddSection("GetAssetActions");
+			Section.AddDynamicEntry(NAME_None, FNewToolMenuSectionDelegate::CreateLambda([](FToolMenuSection& InSection)
+			{
+				if (const UContentBrowserAssetContextMenuContext* CBContext = UContentBrowserAssetContextMenuContext::FindContextWithAssets(InSection))
+				{
+					if (const FAssetData* SelectedAnimBlueprintPtr = CBContext->GetSingleSelectedAssetOfType(UAnimBlueprint::StaticClass(), EIncludeSubclasses::No))
+					{
+						// Accept (non-interface) template anim BPs or anim BPs with compatible skeletons
+						if (SelectedAnimBlueprintPtr &&
+							UE::AnimBlueprint::GetBlueprintType(*SelectedAnimBlueprintPtr) != BPTYPE_Interface &&
+							(
+								(UE::AnimBlueprint::GetAnimBlueprintTargetSkeleton(*SelectedAnimBlueprintPtr).IsNull() && UE::AnimBlueprint::IsAnimBlueprintTemplate(*SelectedAnimBlueprintPtr))
+								||
+								(!UE::AnimBlueprint::GetAnimBlueprintTargetSkeleton(*SelectedAnimBlueprintPtr).IsNull() /* && AnimBlueprint->TargetSkeleton->GetCompatibleSkeletons().Num() > 0 */)
+							)
+						)
+						{
+							InSection.AddSubMenu(
+								"AnimBlueprint_NewSkeletonChildBlueprint",
+								LOCTEXT("AnimBlueprint_NewSkeletonChildBlueprint", "Create Child Anim Blueprint with Skeleton"),
+								LOCTEXT("AnimBlueprint_NewSkeletonChildBlueprint_Tooltip", "Create a child Anim Blueprint that uses a different compatible skeleton"),
+								FNewToolMenuDelegate::CreateStatic(&BuildNewSkeletonChildBlueprintMenu, *SelectedAnimBlueprintPtr),
+								false,
+								FSlateIcon(FAppStyle::GetAppStyleSetName(), "Icons.Blueprint")
+							);
+						}
+					}
+					
+					if (Algo::NoneOf(CBContext->SelectedAssets, &UE::AnimBlueprint::IsAnimBlueprintTemplate) &&
+						Algo::AllOf(CBContext->SelectedAssets, [](const auto& AssetData) { return BPTYPE_Normal == UE::AnimBlueprint::GetBlueprintType(AssetData); } ))
+					{
+						{
+							const TAttribute<FText> Label = LOCTEXT("AnimBlueprint_FindSkeleton", "Find Skeleton");
+							const TAttribute<FText> ToolTip = LOCTEXT("AnimBlueprint_FindSkeletonTooltip", "Finds the skeleton used by the selected Anim Blueprints in the content browser.");
+							const FSlateIcon Icon = FSlateIcon(FAppStyle::GetAppStyleSetName(), "Icons.Find");
+				
+							FToolUIAction UIAction;
+							UIAction.ExecuteAction = FToolMenuExecuteAction::CreateStatic(&ExecuteFindSkeleton);
+							InSection.AddMenuEntry("AnimBlueprint_FindSkeleton", Label, ToolTip, Icon, UIAction);
+						}
+						{
+							const TAttribute<FText> Label = LOCTEXT("AnimBlueprint_AssignSkeleton", "Assign Skeleton");
+							const TAttribute<FText> ToolTip = LOCTEXT("AnimBlueprint_AssignSkeletonTooltip", "Assigns a skeleton to the selected Animation Blueprint(s).");
+							const FSlateIcon Icon = FSlateIcon(FAppStyle::GetAppStyleSetName(), "Persona.AssetActions.AssignSkeleton");
+				
+							FToolUIAction UIAction;
+							UIAction.ExecuteAction = FToolMenuExecuteAction::CreateStatic(&ExecuteAssignSkeleton);
+							InSection.AddMenuEntry("AnimBlueprint_AssignSkeleton", Label, ToolTip, Icon, UIAction);
+						}
+					}
+				}
+			}));
+		}));
+	});
+}
+
+#undef LOCTEXT_NAMESPACE
diff --git a/Engine/Plugins/Editor/EngineAssetDefinitions/Source/Private/Animation/AssetDefinition_AnimBlueprint.h b/Engine/Plugins/Editor/EngineAssetDefinitions/Source/Private/Animation/AssetDefinition_AnimBlueprint.h
new file mode 100644
--- /dev/null
+++ b/Engine/Plugins/Editor/EngineAssetDefinitions/Source/Private/Animation/AssetDefinition_AnimBlueprint.h
@@ -0,0 +1,46 @@
+// Copyright Epic Games, Inc. All Rights Reserved.
+
+#pragma once
+
+#include "Script/AssetDefinition_Blueprint.h"
+#include "Animation/AnimBlueprint.h"
+
+#include "AssetDefinition_AnimBlueprint.generated.h"
+
+struct FToolMenuContext;
+struct FAssetData;
+class UFactory;
+
+UCLASS()
+class UAssetDefinition_AnimBlueprint : public UAssetDefinition_Blueprint
+{
+	GENERATED_BODY()
+
+public:
+	// UAssetDefinition Implementation
+	virtual FText GetAssetDisplayName() const override { return NSLOCTEXT("AssetTypeActions", "AssetTypeActions_AnimBlueprint", "Animation Blueprint"); }
+	virtual FText GetAssetDisplayName(const FAssetData& AssetData) const override;
+	virtual FLinearColor GetAssetColor() const override { return FLinearColor(FColor(200, 116, 0)); }
+	virtual TSoftClassPtr<UObject> GetAssetClass() const override { return UAnimBlueprint::StaticClass(); }
+	virtual TConstArrayView<FAssetCategoryPath> GetAssetCategories() const override
+	{
+		static const auto Categories = { EAssetCategoryPaths::Animation };
+		return Categories;
+	}
+	
+	virtual EAssetCommandResult PerformAssetDiff(const FAssetDiffArgs& DiffArgs) const override;
+	virtual TSharedPtr<SWidget> GetThumbnailOverlay(const FAssetData& AssetData) const override;
+
+	virtual void BuildFilters(TArray<FAssetFilterData>& OutFilters) const override;
+	
+	virtual EAssetCommandResult OpenAssets(const FAssetOpenArgs& OpenArgs) const override;
+	// UAssetDefinition End
+
+	// UAssetDefinition_Blueprint Implementation
+	virtual UFactory* GetFactoryForBlueprintType(UBlueprint* InBlueprint) const override;
+	// UAssetDefinition_Blueprint End
+};
+
+#if UE_ENABLE_INCLUDE_ORDER_DEPRECATED_IN_5_2
+#include "CoreMinimal.h"
+#endif
diff --git a/Engine/Plugins/Editor/EngineAssetDefinitions/Source/Private/Animation/AssetDefinition_AnimBlueprintGeneratedClass.h b/Engine/Plugins/Editor/EngineAssetDefinitions/Source/Private/Animation/AssetDefinition_AnimBlueprintGeneratedClass.h
new file mode 100644
--- /dev/null
+++ b/Engine/Plugins/Editor/EngineAssetDefinitions/Source/Private/Animation/AssetDefinition_AnimBlueprintGeneratedClass.h
@@ -0,0 +1,27 @@
+// Copyright Epic Games, Inc. All Rights Reserved.
+
+#pragma once
+
+#include "CoreMinimal.h"
+#include "Animation/AnimBlueprintGeneratedClass.h"
+#include "Script/AssetDefinition_BlueprintGeneratedClass.h"
+
+#include "AssetDefinition_AnimBlueprintGeneratedClass.generated.h"
+
+UCLASS()
+class UAssetDefinition_AnimBlueprintGeneratedClass : public UAssetDefinition_BlueprintGeneratedClass
+{
+	GENERATED_BODY()
+
+public:
+	// UAssetDefinition Begin
+	virtual FText GetAssetDisplayName() const override { return NSLOCTEXT("AssetTypeActions", "AssetTypeActions_AnimBlueprintGeneratedClass", "Compiled Anim Blueprint Class"); }
+	virtual FLinearColor GetAssetColor() const override { return FLinearColor(FColor(240, 156, 0)); }
+	virtual TSoftClassPtr<UObject> GetAssetClass() const override { return UAnimBlueprintGeneratedClass::StaticClass(); }
+	virtual TConstArrayView<FAssetCategoryPath> GetAssetCategories() const override
+	{
+		static const auto Categories = { EAssetCategoryPaths::Animation };
+		return Categories;
+	}
+	// UAssetDefinition End
+};
diff --git a/Engine/Plugins/Editor/EngineAssetDefinitions/Source/Private/Script/AssetDefinition_BlueprintGeneratedClass.cpp b/Engine/Plugins/Editor/EngineAssetDefinitions/Source/Private/Script/AssetDefinition_BlueprintGeneratedClass.cpp
new file mode 100644
--- /dev/null
+++ b/Engine/Plugins/Editor/EngineAssetDefinitions/Source/Private/Script/AssetDefinition_BlueprintGeneratedClass.cpp
@@ -0,0 +1,47 @@
+// Copyright Epic Games, Inc. All Rights Reserved.
+
+#include "AssetDefinition_BlueprintGeneratedClass.h"
+#include "Engine/BlueprintGeneratedClass.h"
+#include "ToolMenus.h"
+#include "Misc/PackageName.h"
+#include "Misc/MessageDialog.h"
+#include "Factories/BlueprintFactory.h"
+#include "ThumbnailRendering/SceneThumbnailInfo.h"
+#include "Kismet2/KismetEditorUtilities.h"
+#include "IContentBrowserSingleton.h"
+#include "ContentBrowserModule.h"
+#include "IAssetTools.h"
+
+#define LOCTEXT_NAMESPACE "AssetTypeActions"
+
+TWeakPtr<IClassTypeActions> UAssetDefinition_BlueprintGeneratedClass::GetClassTypeActions(const FAssetData& AssetData) const
+{
+	// Blueprints get the class type actions for their parent native class.
+	// Using asset tags avoids us having to load the blueprint
+	UClass* ParentClass = nullptr;
+	FString ParentClassName;
+	if(!AssetData.GetTagValue(FBlueprintTags::NativeParentClassPath, ParentClassName))
+	{
+		AssetData.GetTagValue(FBlueprintTags::ParentClassPath, ParentClassName);
+	}
+	if (!ParentClassName.IsEmpty())
+	{
+		ParentClass = UClass::TryFindTypeSlow<UClass>(FPackageName::ExportTextPathToObjectPath(ParentClassName));
+	}
+
+	if (ParentClass)
+	{
+		return IAssetTools::Get().GetClassTypeActionsForClass(ParentClass);
+	}
+
+	return nullptr;
+}
+
+// Menu Extensions
+//--------------------------------------------------------------------
+
+namespace MenuExtension_BlueprintGeneratedClass
+{
+}
+
+#undef LOCTEXT_NAMESPACE
diff --git a/Engine/Plugins/Editor/EngineAssetDefinitions/Source/Private/Script/AssetDefinition_BlueprintGeneratedClass.h b/Engine/Plugins/Editor/EngineAssetDefinitions/Source/Private/Script/AssetDefinition_BlueprintGeneratedClass.h
new file mode 100644
--- /dev/null
+++ b/Engine/Plugins/Editor/EngineAssetDefinitions/Source/Private/Script/AssetDefinition_BlueprintGeneratedClass.h
@@ -0,0 +1,31 @@
+// Copyright Epic Games, Inc. All Rights Reserved.
+
+#pragma once
+
+#include "CoreMinimal.h"
+#include "Engine/BlueprintGeneratedClass.h"
+#include "Script/AssetDefinition_ClassTypeBase.h"
+
+#include "AssetDefinition_BlueprintGeneratedClass.generated.h"
+
+UCLASS()
+class UAssetDefinition_BlueprintGeneratedClass : public UAssetDefinition_ClassTypeBase
+{
+	GENERATED_BODY()
+
+public:
+	// UAssetDefinition Begin
+	virtual FText GetAssetDisplayName() const override { return NSLOCTEXT("AssetTypeActions", "AssetTypeActions_BlueprintGeneratedClass", "Compiled Blueprint Class"); }
+	virtual FLinearColor GetAssetColor() const override { return FLinearColor(FColor(133, 173, 255)); }
+	virtual TSoftClassPtr<UObject> GetAssetClass() const override { return UBlueprintGeneratedClass::StaticClass(); }
+	virtual TConstArrayView<FAssetCategoryPath> GetAssetCategories() const override
+	{
+		static const TArray<FAssetCategoryPath> Categories = { };
+		return Categories;
+	}
+	// UAssetDefinition End
+
+	// UAssetDefinition_ClassTypeBase Implementation
+	virtual TWeakPtr<IClassTypeActions> GetClassTypeActions(const FAssetData& AssetData) const override;
+	// End UAssetDefinition_ClassTypeBase Implementation
+};
diff --git a/Engine/Source/Developer/AssetTools/Private/AssetDefinition_AssetTypeActionsProxy.cpp b/Engine/Source/Developer/AssetTools/Private/AssetDefinition_AssetTypeActionsProxy.cpp
--- a/Engine/Source/Developer/AssetTools/Private/AssetDefinition_AssetTypeActionsProxy.cpp
+++ b/Engine/Source/Developer/AssetTools/Private/AssetDefinition_AssetTypeActionsProxy.cpp
@@ -193,9 +193,17 @@ void UAssetDefinition_AssetTypeActionsProxy::BuildFilters(TArray<FAssetFilterDat
 {
 	if (AssetType->CanFilter())
 	{
+		// If this asset definition doesn't have any categories it can't have any filters.  Filters need to have a
+		// category to be displayed.
+		if (GetAssetCategories().Num() == 0)
+		{
+			return;
+		}
+		
 		FAssetFilterData Data;
 		Data.Name = AssetType->GetFilterName().ToString();
 		Data.DisplayText = AssetType->GetName();
+		Data.FilterCategories = GetAssetCategories();
 		AssetType->BuildBackendFilter(Data.Filter);
 		OutFilters.Add(MoveTemp(Data));
 	}
diff --git a/Engine/Source/Developer/AssetTools/Private/AssetTools.cpp b/Engine/Source/Developer/AssetTools/Private/AssetTools.cpp
--- a/Engine/Source/Developer/AssetTools/Private/AssetTools.cpp
+++ b/Engine/Source/Developer/AssetTools/Private/AssetTools.cpp
@@ -35,11 +35,6 @@
 #include "IClassTypeActions.h"
 #include "ClassViewerModule.h"
 #include "ClassViewerFilter.h"
-#include "AssetTypeActions/AssetTypeActions_BlueprintGeneratedClass.h"
-#include "AssetTypeActions/AssetTypeActions_AnimationAsset.h"
-#include "AssetTypeActions/AssetTypeActions_AnimBlueprint.h"
-#include "AssetTypeActions/AssetTypeActions_AnimBlueprintInterface.h"
-#include "AssetTypeActions/AssetTypeActions_AnimBlueprintGeneratedClass.h"
 #include "AssetTypeActions/AssetTypeActions_AnimBoneCompressionSettings.h"
 #include "AssetTypeActions/AssetTypeActions_AnimCurveCompressionSettings.h"
 #include "AssetTypeActions/AssetTypeActions_ForceFeedbackEffect.h"
@@ -277,19 +272,19 @@ public:
 
 	virtual bool CanFilter() override
 	{
-		TSharedRef<FAssetFilterDataCache> FilterCache = AssetDefinitionPtr->GetFilters();
+		const TSharedRef<FAssetFilterDataCache> FilterCache = AssetDefinitionPtr->GetFilters();
 		return FilterCache->Filters.Num() > 0;
 	}
 
 	virtual FName GetFilterName() const override
 	{
-		TSharedRef<FAssetFilterDataCache> FilterCache = AssetDefinitionPtr->GetFilters();
+		const TSharedRef<FAssetFilterDataCache> FilterCache = AssetDefinitionPtr->GetFilters();
 		return FilterCache->Filters.Num() > 0 ? FName(*FilterCache->Filters[0].Name) : NAME_None;
 	}
 
 	virtual void BuildBackendFilter(FARFilter& InFilter) override
 	{
-		TSharedRef<FAssetFilterDataCache> FilterCache = AssetDefinitionPtr->GetFilters();
+		const TSharedRef<FAssetFilterDataCache> FilterCache = AssetDefinitionPtr->GetFilters();
 		if (FilterCache->Filters.Num() > 0)
 		{
 			InFilter = FilterCache->Filters[0].Filter;
@@ -1311,25 +1306,21 @@ UAssetToolsImpl::UAssetToolsImpl(const FObjectInitializer& ObjectInitializer)
 
 	// Register the built-in advanced categories
 	AllocatedCategoryBits.Add(TEXT("Animation"), FAdvancedAssetCategory(EAssetTypeCategories::Animation, LOCTEXT("AnimationAssetCategory", "Animation")));
-	AllocatedCategoryBits.Add(TEXT("Blueprint"), FAdvancedAssetCategory(EAssetTypeCategories::Blueprint, LOCTEXT("BlueprintAssetCategory", "Blueprints")));
-	AllocatedCategoryBits.Add(TEXT("Material"), FAdvancedAssetCategory(EAssetTypeCategories::Materials, LOCTEXT("MaterialAssetCategory", "Materials")));
-	AllocatedCategoryBits.Add(TEXT("Audio"), FAdvancedAssetCategory(EAssetTypeCategories::Sounds, LOCTEXT("SoundAssetCategory", "Sounds")));
+	AllocatedCategoryBits.Add(TEXT("Blueprint"), FAdvancedAssetCategory(EAssetTypeCategories::Blueprint, LOCTEXT("BlueprintAssetCategory", "Blueprint")));
+	AllocatedCategoryBits.Add(TEXT("Material"), FAdvancedAssetCategory(EAssetTypeCategories::Materials, LOCTEXT("MaterialAssetCategory", "Material")));
+	AllocatedCategoryBits.Add(TEXT("Audio"), FAdvancedAssetCategory(EAssetTypeCategories::Sounds, LOCTEXT("SoundAssetCategory", "Audio")));
 	AllocatedCategoryBits.Add(TEXT("Physics"), FAdvancedAssetCategory(EAssetTypeCategories::Physics, LOCTEXT("PhysicsAssetCategory", "Physics")));
 	AllocatedCategoryBits.Add(TEXT("User Interface"), FAdvancedAssetCategory(EAssetTypeCategories::UI, LOCTEXT("UserInterfaceAssetCategory", "User Interface")));
 	AllocatedCategoryBits.Add(TEXT("Misc"), FAdvancedAssetCategory(EAssetTypeCategories::Misc, LOCTEXT("MiscellaneousAssetCategory", "Miscellaneous")));
 	AllocatedCategoryBits.Add(TEXT("Gameplay"), FAdvancedAssetCategory(EAssetTypeCategories::Gameplay, LOCTEXT("GameplayAssetCategory", "Gameplay")));
 	AllocatedCategoryBits.Add(TEXT("Media"), FAdvancedAssetCategory(EAssetTypeCategories::Media, LOCTEXT("MediaAssetCategory", "Media")));
-	AllocatedCategoryBits.Add(TEXT("Texture"), FAdvancedAssetCategory(EAssetTypeCategories::Textures, LOCTEXT("TextureAssetCategory", "Textures")));
+	AllocatedCategoryBits.Add(TEXT("Texture"), FAdvancedAssetCategory(EAssetTypeCategories::Textures, LOCTEXT("TextureAssetCategory", "Texturs")));
 
 	EAssetTypeCategories::Type InputCategoryBit = RegisterAdvancedAssetCategory(FName(TEXT("Input")), LOCTEXT("InputAssetsCategory", "Input"));
 	
 	// Register the built-in asset type actions
-	RegisterAssetTypeActions(MakeShareable(new FAssetTypeActions_AnimBlueprintInterface));
-	RegisterAssetTypeActions(MakeShareable(new FAssetTypeActions_AnimBlueprint));
-	RegisterAssetTypeActions(MakeShareable(new FAssetTypeActions_AnimBlueprintGeneratedClass));
 	RegisterAssetTypeActions(MakeShareable(new FAssetTypeActions_AnimBoneCompressionSettings));
 	RegisterAssetTypeActions(MakeShareable(new FAssetTypeActions_AnimCurveCompressionSettings));
-	RegisterAssetTypeActions(MakeShareable(new FAssetTypeActions_BlueprintGeneratedClass));
 	RegisterAssetTypeActions(MakeShareable(new FAssetTypeActions_ForceFeedbackEffect(InputCategoryBit)));
 	RegisterAssetTypeActions(MakeShareable(new FAssetTypeActions_ParticleSystem));
 	RegisterAssetTypeActions(MakeShareable(new FAssetTypeActions_PhysicalMaterialMask));
diff --git a/Engine/Source/Developer/AssetTools/Private/AssetTypeActions/AssetTypeActions_AnimBlueprint.cpp b/Engine/Source/Developer/AssetTools/Private/AssetTypeActions/AssetTypeActions_AnimBlueprint.cpp
deleted file mode 100644
diff --git a/Engine/Source/Developer/AssetTools/Private/AssetTypeActions/AssetTypeActions_AnimBlueprint.h b/Engine/Source/Developer/AssetTools/Private/AssetTypeActions/AssetTypeActions_AnimBlueprint.h
deleted file mode 100644
diff --git a/Engine/Source/Developer/AssetTools/Private/AssetTypeActions/AssetTypeActions_AnimBlueprintGeneratedClass.h b/Engine/Source/Developer/AssetTools/Private/AssetTypeActions/AssetTypeActions_AnimBlueprintGeneratedClass.h
deleted file mode 100644
diff --git a/Engine/Source/Developer/AssetTools/Private/AssetTypeActions/AssetTypeActions_AnimBlueprintInterface.cpp b/Engine/Source/Developer/AssetTools/Private/AssetTypeActions/AssetTypeActions_AnimBlueprintInterface.cpp
deleted file mode 100644
diff --git a/Engine/Source/Developer/AssetTools/Private/AssetTypeActions/AssetTypeActions_AnimBlueprintInterface.h b/Engine/Source/Developer/AssetTools/Private/AssetTypeActions/AssetTypeActions_AnimBlueprintInterface.h
deleted file mode 100644
diff --git a/Engine/Source/Developer/AssetTools/Private/AssetTypeActions/AssetTypeActions_BlueprintGeneratedClass.cpp b/Engine/Source/Developer/AssetTools/Private/AssetTypeActions/AssetTypeActions_BlueprintGeneratedClass.cpp
--- a/Engine/Source/Developer/AssetTools/Private/AssetTypeActions/AssetTypeActions_BlueprintGeneratedClass.cpp
+++ b/Engine/Source/Developer/AssetTools/Private/AssetTypeActions/AssetTypeActions_BlueprintGeneratedClass.cpp
@@ -18,94 +18,6 @@ UClass* FAssetTypeActions_BlueprintGeneratedClass::GetSupportedClass() const
 	return UBlueprintGeneratedClass::StaticClass();
 }
 
-void FAssetTypeActions_BlueprintGeneratedClass::GetActions(const TArray<UObject*>& InObjects, FToolMenuSection& Section)
-{
-	TArray<TWeakObjectPtr<UBlueprintGeneratedClass>> BPGCs = GetTypedWeakObjectPtrs<UBlueprintGeneratedClass>(InObjects);
-
-	if (BPGCs.Num() == 1)
-	{
-		TAttribute<FText>::FGetter DynamicTooltipGetter;
-		DynamicTooltipGetter.BindSP(this, &FAssetTypeActions_BlueprintGeneratedClass::GetNewDerivedBlueprintTooltip, BPGCs[0]);
-		TAttribute<FText> DynamicTooltipAttribute = TAttribute<FText>::Create(DynamicTooltipGetter);
-
-		Section.AddMenuEntry(
-			"BlueprintGeneratedClass_NewDerivedBlueprint",
-			LOCTEXT("BlueprintGeneratedClass_NewDerivedBlueprint", "Create Child Blueprint Class"),
-			DynamicTooltipAttribute,
-			FSlateIcon(FAppStyle::GetAppStyleSetName(), "Icons.Blueprint"),
-			FUIAction(
-				FExecuteAction::CreateSP(this, &FAssetTypeActions_BlueprintGeneratedClass::ExecuteNewDerivedBlueprint, BPGCs[0]),
-				FCanExecuteAction::CreateSP(this, &FAssetTypeActions_BlueprintGeneratedClass::CanExecuteNewDerivedBlueprint, BPGCs[0])
-			)
-		);
-	}
-}
-
-UFactory* FAssetTypeActions_BlueprintGeneratedClass::GetFactoryForNewDerivedBlueprint(UBlueprintGeneratedClass* InBPGC) const
-{
-	UBlueprintFactory* Factory = NewObject<UBlueprintFactory>();
-	Factory->ParentClass = InBPGC;
-	return Factory;
-}
-
-UClass* FAssetTypeActions_BlueprintGeneratedClass::GetNewDerivedBlueprintClass() const
-{
-	return UBlueprint::StaticClass();
-}
-
-void FAssetTypeActions_BlueprintGeneratedClass::ExecuteNewDerivedBlueprint(TWeakObjectPtr<UBlueprintGeneratedClass> InObject)
-{
-	if (UBlueprintGeneratedClass* TargetParentClass = InObject.Get())
-	{
-		if (!FKismetEditorUtilities::CanCreateBlueprintOfClass(TargetParentClass))
-		{
-			FMessageDialog::Open(EAppMsgType::Ok, LOCTEXT("InvalidClassToMakeBlueprintFrom", "Invalid class with which to make a Blueprint."));
-			return;
-		}
-
-		FString Name;
-		FString PackageName;
-		CreateUniqueAssetName(TargetParentClass->GetOutermost()->GetName(), TEXT("_Child"), PackageName, Name);
-		const FString PackagePath = FPackageName::GetLongPackagePath(PackageName);
-
-		FContentBrowserModule& ContentBrowserModule = FModuleManager::LoadModuleChecked<FContentBrowserModule>("ContentBrowser");
-		ContentBrowserModule.Get().CreateNewAsset(Name, PackagePath, GetNewDerivedBlueprintClass(), GetFactoryForNewDerivedBlueprint(TargetParentClass));
-	}
-}
-
-FText FAssetTypeActions_BlueprintGeneratedClass::GetNewDerivedBlueprintTooltip(TWeakObjectPtr<UBlueprintGeneratedClass> InObject)
-{
-	if (!CanExecuteNewDerivedBlueprint(InObject))
-	{
-		return LOCTEXT("BlueprintGeneratedClass_NewDerivedBlueprintIsDeprecatedTooltip", "Compiled Blueprint class is deprecated, cannot derive a child Blueprint!");
-	}
-	else
-	{
-		return LOCTEXT("BlueprintGeneratedClass_NewDerivedBlueprintTooltip", "Creates a Child Blueprint Class based on the current compiled Blueprint, allowing you to create variants easily.");
-	}
-}
-
-bool FAssetTypeActions_BlueprintGeneratedClass::CanExecuteNewDerivedBlueprint(TWeakObjectPtr<UBlueprintGeneratedClass> InObject)
-{
-	UBlueprintGeneratedClass* BPGC = InObject.Get();
-	return BPGC && !BPGC->HasAnyClassFlags(CLASS_Deprecated);
-}
-
-//@todo
-
-//UThumbnailInfo* FAssetTypeActions_BlueprintGeneratedClass::GetThumbnailInfo(UObject* Asset) const
-//{
-//	// Blueprint thumbnail scenes are disabled for now
-//	UBlueprintGeneratedClass* BPGC = CastChecked<UBlueprintGeneratedClass>(Asset);
-//	UThumbnailInfo* ThumbnailInfo = BPGC->ThumbnailInfo;
-//	if (!ThumbnailInfo)
-//	{
-//		ThumbnailInfo = NewObject<USceneThumbnailInfo>(BPGC, NAME_None, RF_Transactional);
-//		BPGC->ThumbnailInfo = ThumbnailInfo;
-//	}
-//	return ThumbnailInfo;
-//}
-
 TWeakPtr<IClassTypeActions> FAssetTypeActions_BlueprintGeneratedClass::GetClassTypeActions(const FAssetData& AssetData) const
 {
 	// Blueprints get the class type actions for their parent native class.
diff --git a/Engine/Source/Developer/AssetTools/Public/AssetTypeActions/AssetTypeActions_BlueprintGeneratedClass.h b/Engine/Source/Developer/AssetTools/Public/AssetTypeActions/AssetTypeActions_BlueprintGeneratedClass.h
--- a/Engine/Source/Developer/AssetTools/Public/AssetTypeActions/AssetTypeActions_BlueprintGeneratedClass.h
+++ b/Engine/Source/Developer/AssetTools/Public/AssetTypeActions/AssetTypeActions_BlueprintGeneratedClass.h
@@ -11,14 +11,15 @@ struct FAssetData;
 class IClassTypeActions;
 class UFactory;
 
-class ASSETTOOLS_API FAssetTypeActions_BlueprintGeneratedClass : public FAssetTypeActions_ClassTypeBase
+class
+// UE_DEPRECATED(5.2, "The AssetDefinition system is replacing AssetTypeActions and UAssetDefinition_BlueprintGeneratedClass replaced this.  Please see the Conversion Guide in AssetDefinition.h")
+ASSETTOOLS_API FAssetTypeActions_BlueprintGeneratedClass : public FAssetTypeActions_ClassTypeBase
 {
 public:
 	// IAssetTypeActions Implementation
 	virtual FText GetName() const override { return NSLOCTEXT("AssetTypeActions", "AssetTypeActions_BlueprintGeneratedClass", "Compiled Blueprint Class"); }
 	virtual FColor GetTypeColor() const override { return FColor(133, 173, 255); }
 	virtual UClass* GetSupportedClass() const override;
-	virtual void GetActions(const TArray<UObject*>& InObjects, FToolMenuSection& Section) override;
 
 	// We don't want BPGCs to show up in CB/Ref viewer filters since the BP filters already include them
 	virtual uint32 GetCategories() override { return EAssetTypeCategories::None; }
@@ -29,21 +30,4 @@ public:
 	// FAssetTypeActions_ClassTypeBase Implementation
 	virtual TWeakPtr<IClassTypeActions> GetClassTypeActions(const FAssetData& AssetData) const override;
 	// End FAssetTypeActions_ClassTypeBase Implementation
-
-protected:
-	/** Returns the factory responsible for creating a new Blueprint derived from the specified BPGC */
-	virtual UFactory* GetFactoryForNewDerivedBlueprint(UBlueprintGeneratedClass* InBPGC) const;
-
-	/** Returns the class of new Blueprints derived from this BPGC asset type */
-	virtual UClass* GetNewDerivedBlueprintClass() const;
-
-	/** Returns the tooltip to display when attempting to derive a Blueprint */
-	FText GetNewDerivedBlueprintTooltip(TWeakObjectPtr<UBlueprintGeneratedClass> InObject);
-
-	/** Returns TRUE if you can derive a Blueprint */
-	bool CanExecuteNewDerivedBlueprint(TWeakObjectPtr<UBlueprintGeneratedClass> InObject);
-
-private:
-	/** Handler for when NewDerivedBlueprint is selected */
-	void ExecuteNewDerivedBlueprint(TWeakObjectPtr<UBlueprintGeneratedClass> InObject);
 };
diff --git a/Engine/Source/Editor/AssetDefinition/Private/AssetDefinition.cpp b/Engine/Source/Editor/AssetDefinition/Private/AssetDefinition.cpp
--- a/Engine/Source/Editor/AssetDefinition/Private/AssetDefinition.cpp
+++ b/Engine/Source/Editor/AssetDefinition/Private/AssetDefinition.cpp
@@ -102,6 +102,13 @@ void UAssetDefinition::BuildFilters(TArray<FAssetFilterData>& OutFilters) const
 
 	if (const UClass* AssetClass = AssetClassPtr.Get())
 	{
+		// If this asset definition doesn't have any categories it can't have any filters.  Filters need to have a
+		// category to be displayed.
+		if (GetAssetCategories().Num() == 0)
+		{
+			return;
+		}
+		
 		// By default we don't advertise filtering if the class is abstract for the asset definition.  Odds are,
 		// if they've registered an abstract class as an asset definition, they mean to use it for subclasses.
 		if (IncludeClassInFilter == EIncludeClassInFilter::Always || (IncludeClassInFilter == EIncludeClassInFilter::IfClassIsNotAbstract && !AssetClass->HasAnyClassFlags(CLASS_Abstract)))
diff --git a/Engine/Source/Editor/EditorWidgets/Public/Filters/SAssetFilterBar.h b/Engine/Source/Editor/EditorWidgets/Public/Filters/SAssetFilterBar.h
--- a/Engine/Source/Editor/EditorWidgets/Public/Filters/SAssetFilterBar.h
+++ b/Engine/Source/Editor/EditorWidgets/Public/Filters/SAssetFilterBar.h
@@ -834,6 +834,8 @@ protected:
 				
 					for (const FAssetFilterData& FilterData : FilterCache->Filters)
 					{
+						ensureMsgf(FilterData.FilterCategories.Num() > 0, TEXT("%s is missing Filter Categories, without any filter categories we can't display this filter."), *FilterData.Name);
+					
 						// Convert the AssetTypeAction to an FCustomClassFilterData and add it to our list
 						TSharedRef<FCustomClassFilterData> CustomClassFilterData = MakeShared<FCustomClassFilterData>(AssetDefinition, FilterData);
 
