AssetDefinition - Fix the AnimSequence color.

diff --git a/Engine/Plugins/Editor/EngineAssetDefinitions/Source/Private/Animation/AssetDefinition_AnimSequence.h b/Engine/Plugins/Editor/EngineAssetDefinitions/Source/Private/Animation/AssetDefinition_AnimSequence.h
--- a/Engine/Plugins/Editor/EngineAssetDefinitions/Source/Private/Animation/AssetDefinition_AnimSequence.h
+++ b/Engine/Plugins/Editor/EngineAssetDefinitions/Source/Private/Animation/AssetDefinition_AnimSequence.h
@@ -15,7 +15,6 @@ class UAssetDefinition_AnimSequence : public UAssetDefinition_AnimationAsset
 public:
 	// UAssetDefinition Begin
 	virtual FText GetAssetDisplayName() const override { return NSLOCTEXT("AssetTypeActions", "AssetTypeActions_AnimSequence", "Animation Sequence"); }
-	virtual FLinearColor GetAssetColor() const override { return FLinearColor(FColor(237, 28, 36)); }
 	virtual TSoftClassPtr<UObject> GetAssetClass() const override { return UAnimSequence::StaticClass(); }
 	virtual bool CanImport() const override { return true; }
 	// UAssetDefinition End
