AssetDefintiion - Fixing a couple of blueprint related asset definition items.

diff --git a/Engine/Plugins/Editor/EngineAssetDefinitions/Source/Private/Animation/AssetDefinition_AnimBlueprint.cpp b/Engine/Plugins/Editor/EngineAssetDefinitions/Source/Private/Animation/AssetDefinition_AnimBlueprint.cpp
--- a/Engine/Plugins/Editor/EngineAssetDefinitions/Source/Private/Animation/AssetDefinition_AnimBlueprint.cpp
+++ b/Engine/Plugins/Editor/EngineAssetDefinitions/Source/Private/Animation/AssetDefinition_AnimBlueprint.cpp
@@ -332,7 +332,7 @@ namespace MenuExtension_AnimBlueprint
 		UToolMenus::RegisterStartupCallback(FSimpleMulticastDelegate::FDelegate::CreateLambda([]()
 		{
 			FToolMenuOwnerScoped OwnerScoped(UE_MODULE_NAME);
-			UToolMenu* Menu = UE::ContentBrowser::ExtendToolMenu_AssetContextMenu(UBlueprint::StaticClass());
+			UToolMenu* Menu = UE::ContentBrowser::ExtendToolMenu_AssetContextMenu(UAnimBlueprint::StaticClass());
 	        
 			FToolMenuSection& Section = Menu->FindOrAddSection("GetAssetActions");
 			Section.AddDynamicEntry(NAME_None, FNewToolMenuSectionDelegate::CreateLambda([](FToolMenuSection& InSection)
diff --git a/Engine/Plugins/Editor/EngineAssetDefinitions/Source/Private/Script/AssetDefinition_BlueprintGeneratedClass.cpp b/Engine/Plugins/Editor/EngineAssetDefinitions/Source/Private/Script/AssetDefinition_BlueprintGeneratedClass.cpp
--- a/Engine/Plugins/Editor/EngineAssetDefinitions/Source/Private/Script/AssetDefinition_BlueprintGeneratedClass.cpp
+++ b/Engine/Plugins/Editor/EngineAssetDefinitions/Source/Private/Script/AssetDefinition_BlueprintGeneratedClass.cpp
@@ -1,12 +1,14 @@
 // Copyright Epic Games, Inc. All Rights Reserved.
 
 #include "AssetDefinition_BlueprintGeneratedClass.h"
+
+#include "AssetDefinitionRegistry.h"
+#include "ContentBrowserMenuContexts.h"
 #include "Engine/BlueprintGeneratedClass.h"
 #include "ToolMenus.h"
 #include "Misc/PackageName.h"
 #include "Misc/MessageDialog.h"
 #include "Factories/BlueprintFactory.h"
-#include "ThumbnailRendering/SceneThumbnailInfo.h"
 #include "Kismet2/KismetEditorUtilities.h"
 #include "IContentBrowserSingleton.h"
 #include "ContentBrowserModule.h"
@@ -37,11 +39,104 @@ TWeakPtr<IClassTypeActions> UAssetDefinition_BlueprintGeneratedClass::GetClassTy
 	return nullptr;
 }
 
+UClass* UAssetDefinition_BlueprintGeneratedClass::GetNewDerivedBlueprintClass() const
+{
+	return UBlueprint::StaticClass();
+}
+
+UFactory* UAssetDefinition_BlueprintGeneratedClass::GetFactoryForNewDerivedBlueprint(UBlueprintGeneratedClass* GeneratedClass) const
+{
+	UBlueprintFactory* BlueprintFactory = NewObject<UBlueprintFactory>();
+	BlueprintFactory->ParentClass = GeneratedClass;
+	return BlueprintFactory;
+}
+
 // Menu Extensions
 //--------------------------------------------------------------------
 
 namespace MenuExtension_BlueprintGeneratedClass
 {
+	static bool CanExecuteNewDerivedBlueprint(const FToolMenuContext& MenuContext, FAssetData GeneratedClassAsset)
+	{
+		const uint32 BPFlags = GeneratedClassAsset.GetTagValueRef<uint32>(FBlueprintTags::ClassFlags);
+		if ((BPFlags & (CLASS_Deprecated)) == 0)
+		{
+			return true;
+		}
+
+		return false;
+	}
+
+	static void ExecuteNewDerivedBlueprint(const FToolMenuContext& MenuContext, FAssetData GeneratedClassAsset)
+	{
+		if (UBlueprintGeneratedClass* TargetParentClass = Cast<UBlueprintGeneratedClass>(GeneratedClassAsset.GetAsset()))
+		{
+			if (!FKismetEditorUtilities::CanCreateBlueprintOfClass(TargetParentClass))
+			{
+				FMessageDialog::Open( EAppMsgType::Ok, LOCTEXT("InvalidClassToMakeBlueprintFrom", "Invalid class with which to make a Blueprint."));
+				return;
+			}
+
+			FString Name;
+			FString PackageName;
+			IAssetTools::Get().CreateUniqueAssetName(TargetParentClass->GetOutermost()->GetName(), TEXT("_Child"), PackageName, Name);
+			const FString PackagePath = FPackageName::GetLongPackagePath(PackageName);
+
+			const UAssetDefinition_BlueprintGeneratedClass* BPGC_AssetDefinition = Cast<UAssetDefinition_BlueprintGeneratedClass>(UAssetDefinitionRegistry::Get()->GetAssetDefinitionForClass(TargetParentClass->GetClass()));
+			if (BPGC_AssetDefinition)
+			{
+				UFactory* AssetFactory = BPGC_AssetDefinition->GetFactoryForNewDerivedBlueprint(TargetParentClass);
+				UClass* AssetClass = BPGC_AssetDefinition->GetNewDerivedBlueprintClass();
+
+				FContentBrowserModule& ContentBrowserModule = FModuleManager::LoadModuleChecked<FContentBrowserModule>("ContentBrowser");
+				ContentBrowserModule.Get().CreateNewAsset(Name, PackagePath, AssetClass, AssetFactory);
+			}
+		}
+	}
+	
+	static FDelayedAutoRegisterHelper DelayedAutoRegister(EDelayedRegisterRunPhase::EndOfEngineInit, []{ 
+		UToolMenus::RegisterStartupCallback(FSimpleMulticastDelegate::FDelegate::CreateLambda([]()
+		{
+			FToolMenuOwnerScoped OwnerScoped(UE_MODULE_NAME);
+			UToolMenu* Menu = UE::ContentBrowser::ExtendToolMenu_AssetContextMenu(UBlueprintGeneratedClass::StaticClass());
+		
+			FToolMenuSection& Section = Menu->FindOrAddSection("GetAssetActions");
+			Section.AddDynamicEntry(NAME_None, FNewToolMenuSectionDelegate::CreateLambda([](FToolMenuSection& InSection)
+			{
+				if (const UContentBrowserAssetContextMenuContext* CBContext = UContentBrowserAssetContextMenuContext::FindContextWithAssets(InSection))
+				{
+					if (const FAssetData* SelectedBPGCPtr = CBContext->GetSingleSelectedAssetOfType(UBlueprintGeneratedClass::StaticClass()))
+					{
+						const UAssetDefinition_BlueprintGeneratedClass* BPGC_AssetDefinition = Cast<UAssetDefinition_BlueprintGeneratedClass>(UAssetDefinitionRegistry::Get()->GetAssetDefinitionForClass(SelectedBPGCPtr->GetClass()));
+						if (BPGC_AssetDefinition && BPGC_AssetDefinition->GetNewDerivedBlueprintClass() != nullptr)
+						{
+							FAssetData GeneratedClassAsset = *SelectedBPGCPtr;
+							
+							const TAttribute<FText> Label = LOCTEXT("BlueprintGeneratedClass_NewDerivedBlueprint", "Create Child Blueprint Class");
+							const TAttribute<FText> ToolTip = TAttribute<FText>::CreateLambda([GeneratedClassAsset]()
+							{
+								const uint32 BPFlags = GeneratedClassAsset.GetTagValueRef<uint32>(FBlueprintTags::ClassFlags);
+								if ((BPFlags & (CLASS_Deprecated)) == 0)
+								{
+									return LOCTEXT("Blueprint_NewDerivedBlueprintTooltip", "Creates a Child Blueprint Class based on the current Blueprint, allowing you to create variants easily.");
+								}
+								else
+								{
+									return LOCTEXT("Blueprint_NewDerivedBlueprintIsDeprecatedTooltip", "Blueprint class is deprecated, cannot derive a child Blueprint!");
+								}
+							});
+							const FSlateIcon Icon = FSlateIcon(FAppStyle::GetAppStyleSetName(), "ClassIcon.Blueprint");
+
+							FToolUIAction UIAction;
+							UIAction.ExecuteAction = FToolMenuExecuteAction::CreateStatic(&ExecuteNewDerivedBlueprint, GeneratedClassAsset);
+							UIAction.CanExecuteAction = FToolMenuCanExecuteAction::CreateStatic(&CanExecuteNewDerivedBlueprint, GeneratedClassAsset);
+							InSection.AddMenuEntry("BlueprintGeneratedClass_NewDerivedBlueprint", Label, ToolTip, Icon, UIAction);
+						}
+					}
+				}
+			}));
+		}));
+	});
 }
 
 #undef LOCTEXT_NAMESPACE
diff --git a/Engine/Plugins/Editor/EngineAssetDefinitions/Source/Private/Script/AssetDefinition_BlueprintGeneratedClass.h b/Engine/Plugins/Editor/EngineAssetDefinitions/Source/Private/Script/AssetDefinition_BlueprintGeneratedClass.h
--- a/Engine/Plugins/Editor/EngineAssetDefinitions/Source/Private/Script/AssetDefinition_BlueprintGeneratedClass.h
+++ b/Engine/Plugins/Editor/EngineAssetDefinitions/Source/Private/Script/AssetDefinition_BlueprintGeneratedClass.h
@@ -8,6 +8,8 @@
 
 #include "AssetDefinition_BlueprintGeneratedClass.generated.h"
 
+class UFactory;
+
 UCLASS()
 class UAssetDefinition_BlueprintGeneratedClass : public UAssetDefinition_ClassTypeBase
 {
@@ -28,4 +30,7 @@ public:
 	// UAssetDefinition_ClassTypeBase Implementation
 	virtual TWeakPtr<IClassTypeActions> GetClassTypeActions(const FAssetData& AssetData) const override;
 	// End UAssetDefinition_ClassTypeBase Implementation
+
+	virtual UClass* GetNewDerivedBlueprintClass() const;
+	virtual UFactory* GetFactoryForNewDerivedBlueprint(UBlueprintGeneratedClass* GeneratedClass) const;
 };
