AssetDefinition - Fixing AssetDefintions for native classes not existing.

diff --git a/Engine/Plugins/Editor/ContentBrowser/ContentBrowserClassDataSource/Source/ContentBrowserClassDataSource/Private/ContentBrowserClassDataSource.cpp b/Engine/Plugins/Editor/ContentBrowser/ContentBrowserClassDataSource/Source/ContentBrowserClassDataSource/Private/ContentBrowserClassDataSource.cpp
--- a/Engine/Plugins/Editor/ContentBrowser/ContentBrowserClassDataSource/Source/ContentBrowserClassDataSource/Private/ContentBrowserClassDataSource.cpp
+++ b/Engine/Plugins/Editor/ContentBrowser/ContentBrowserClassDataSource/Source/ContentBrowserClassDataSource/Private/ContentBrowserClassDataSource.cpp
@@ -30,12 +30,6 @@ void UContentBrowserClassDataSource::Initialize(const bool InAutoRegister)
 {
 	Super::Initialize(InAutoRegister);
 
-	{
-		static const FName NAME_AssetTools = "AssetTools";
-		FAssetToolsModule& AssetToolsModule = FModuleManager::LoadModuleChecked<FAssetToolsModule>(NAME_AssetTools);
-		ClassTypeActions = AssetToolsModule.Get().GetAssetTypeActionsForClass(UClass::StaticClass()).Pin();
-	}
-
 	CollectionManager = &FCollectionManagerModule::GetModule().Get();
 
 	// Bind the class specific menu extensions
@@ -64,6 +58,18 @@ void UContentBrowserClassDataSource::Shutdown()
 	Super::Shutdown();
 }
 
+TSharedPtr<IAssetTypeActions> UContentBrowserClassDataSource::GetClassTypeActions()
+{
+	if (!ClassTypeActions)
+	{
+		static const FName NAME_AssetTools = "AssetTools";
+		FAssetToolsModule& AssetToolsModule = FModuleManager::LoadModuleChecked<FAssetToolsModule>(NAME_AssetTools);
+		ClassTypeActions = AssetToolsModule.Get().GetAssetTypeActionsForClass(UClass::StaticClass()).Pin();
+	}
+
+	return ClassTypeActions;
+}
+
 void UContentBrowserClassDataSource::BuildRootPathVirtualTree()
 {
 	Super::BuildRootPathVirtualTree();
@@ -477,7 +483,7 @@ bool UContentBrowserClassDataSource::DoesItemPassFilter(const FContentBrowserIte
 
 bool UContentBrowserClassDataSource::GetItemAttribute(const FContentBrowserItemData& InItem, const bool InIncludeMetaData, const FName InAttributeKey, FContentBrowserItemDataAttributeValue& OutAttributeValue)
 {
-	return ContentBrowserClassData::GetItemAttribute(ClassTypeActions.Get(), this, InItem, InIncludeMetaData, InAttributeKey, OutAttributeValue);
+	return ContentBrowserClassData::GetItemAttribute(GetClassTypeActions().Get(), this, InItem, InIncludeMetaData, InAttributeKey, OutAttributeValue);
 }
 
 bool UContentBrowserClassDataSource::GetItemAttributes(const FContentBrowserItemData& InItem, const bool InIncludeMetaData, FContentBrowserItemDataAttributeValues& OutAttributeValues)
@@ -497,12 +503,12 @@ bool UContentBrowserClassDataSource::CanEditItem(const FContentBrowserItemData&
 
 bool UContentBrowserClassDataSource::EditItem(const FContentBrowserItemData& InItem)
 {
-	return ContentBrowserClassData::EditItems(ClassTypeActions.Get(), this, MakeArrayView(&InItem, 1));
+	return ContentBrowserClassData::EditItems(GetClassTypeActions().Get(), this, MakeArrayView(&InItem, 1));
 }
 
 bool UContentBrowserClassDataSource::BulkEditItems(TArrayView<const FContentBrowserItemData> InItems)
 {
-	return ContentBrowserClassData::EditItems(ClassTypeActions.Get(), this, InItems);
+	return ContentBrowserClassData::EditItems(GetClassTypeActions().Get(), this, InItems);
 }
 
 bool UContentBrowserClassDataSource::AppendItemReference(const FContentBrowserItemData& InItem, FString& InOutStr)
diff --git a/Engine/Plugins/Editor/ContentBrowser/ContentBrowserClassDataSource/Source/ContentBrowserClassDataSource/Public/ContentBrowserClassDataSource.h b/Engine/Plugins/Editor/ContentBrowser/ContentBrowserClassDataSource/Source/ContentBrowserClassDataSource/Public/ContentBrowserClassDataSource.h
--- a/Engine/Plugins/Editor/ContentBrowser/ContentBrowserClassDataSource/Source/ContentBrowserClassDataSource/Public/ContentBrowserClassDataSource.h
+++ b/Engine/Plugins/Editor/ContentBrowser/ContentBrowserClassDataSource/Source/ContentBrowserClassDataSource/Public/ContentBrowserClassDataSource.h
@@ -78,6 +78,8 @@ public:
 	virtual bool Legacy_TryConvertAssetDataToVirtualPath(const FAssetData& InAssetData, const bool InUseFolderPaths, FName& OutPath) override;
 
 protected:
+	TSharedPtr<IAssetTypeActions> GetClassTypeActions();
+	
 	virtual void BuildRootPathVirtualTree() override;
 	bool RootClassPathPassesFilter(const FName InRootClassPath, const bool bIncludeEngineClasses, const bool bIncludePluginClasses) const;
 
diff --git a/Engine/Plugins/Editor/EngineAssetDefinitions/Source/Private/Script/AssetDefinition_Class.cpp b/Engine/Plugins/Editor/EngineAssetDefinitions/Source/Private/Script/AssetDefinition_Class.cpp
--- a/Engine/Plugins/Editor/EngineAssetDefinitions/Source/Private/Script/AssetDefinition_Class.cpp
+++ b/Engine/Plugins/Editor/EngineAssetDefinitions/Source/Private/Script/AssetDefinition_Class.cpp
@@ -7,6 +7,8 @@
 #include "Styling/AppStyle.h"
 #include "SourceCodeNavigation.h"
 #include "AddToProjectConfig.h"
+#include "AssetToolsModule.h"
+#include "ContentBrowserMenuContexts.h"
 #include "GameProjectGenerationModule.h"
 #include "Kismet2/KismetEditorUtilities.h"
 #include "Editor/UnrealEdEngine.h"
@@ -15,123 +17,8 @@
 
 #define LOCTEXT_NAMESPACE "AssetTypeActions"
 
-void FAssetTypeActions_Class::GetActions(const TArray<UObject*>& InObjects, FToolMenuSection& Section)
-{
-	UClass *const BaseClass = (InObjects.Num() == 1) ? Cast<UClass>(InObjects[0]) : nullptr;
-
-	// Only allow the New class option if we have a base class that we can actually derive from in one of our project modules
-	FGameProjectGenerationModule& GameProjectGenerationModule = FGameProjectGenerationModule::Get();
-	TArray<FModuleContextInfo> ProjectModules = GameProjectGenerationModule.GetCurrentProjectModules();
-	const bool bIsValidBaseCppClass = BaseClass && GameProjectGenerationModule.IsValidBaseClassForCreation(BaseClass, ProjectModules);
-	const bool bIsValidBaseBlueprintClass = BaseClass && FKismetEditorUtilities::CanCreateBlueprintOfClass(BaseClass);
-
-	auto CreateCreateDerivedCppClass = [BaseClass]()
-	{
-		// Work out where the header file for the current class is, as we'll use that path as the default for the new class
-		FString BaseClassPath;
-		if(FSourceCodeNavigation::FindClassHeaderPath(BaseClass, BaseClassPath))
-		{
-			// Strip off the actual filename as we only need the path
-			BaseClassPath = FPaths::GetPath(BaseClassPath);
-		}
-
-		FGameProjectGenerationModule::Get().OpenAddCodeToProjectDialog(
-			FAddToProjectConfig()
-			.ParentClass(BaseClass)
-			.InitialPath(BaseClassPath)
-			.ParentWindow(FGlobalTabmanager::Get()->GetRootWindow())
-		);
-	};
 
-	auto IsCPPAllowed = []()
-	{
-		return ensure(GUnrealEd) && GUnrealEd->GetUnrealEdOptions()->IsCPPAllowed();
-	};
-
-	auto CanCreateDerivedCppClass = [bIsValidBaseCppClass, IsCPPAllowed]() -> bool
-	{
-		return bIsValidBaseCppClass && IsCPPAllowed();
-	};
-
-	auto CreateCreateDerivedBlueprintClass = [BaseClass]()
-	{
-		FGameProjectGenerationModule::Get().OpenAddBlueprintToProjectDialog(
-			FAddToProjectConfig()
-			.ParentClass(BaseClass)
-			.ParentWindow(FGlobalTabmanager::Get()->GetRootWindow())
-		);
-	};
-
-	auto CanCreateDerivedBlueprintClass = [bIsValidBaseBlueprintClass]() -> bool
-	{
-		return bIsValidBaseBlueprintClass;
-	};
-
-	FText NewDerivedCppClassLabel;
-	FText NewDerivedCppClassToolTip;
-	FText NewDerivedBlueprintClassLabel;
-	FText NewDerivedBlueprintClassToolTip;
-	if(InObjects.Num() == 1)
-	{
-		check(BaseClass);
-
-		const FText BaseClassName = FText::FromName(BaseClass->GetFName());
-
-		NewDerivedCppClassLabel = FText::Format(LOCTEXT("Class_NewDerivedCppClassLabel_CreateFrom", "Create C++ class derived from {0}"), BaseClassName);
-		if(bIsValidBaseCppClass)
-		{
-			NewDerivedCppClassToolTip = FText::Format(LOCTEXT("Class_NewDerivedCppClassTooltip_CreateFrom", "Create a new C++ class deriving from {0}."), BaseClassName);
-		}
-		else
-		{
-			NewDerivedCppClassToolTip = FText::Format(LOCTEXT("Class_NewDerivedCppClassTooltip_InvalidClass", "Cannot create a new C++ class deriving from {0}."), BaseClassName);
-		}
-
-		NewDerivedBlueprintClassLabel = FText::Format(LOCTEXT("Class_NewDerivedBlueprintClassLabel_CreateFrom", "Create Blueprint class based on {0}"), BaseClassName);
-		if(bIsValidBaseBlueprintClass)
-		{
-			NewDerivedBlueprintClassToolTip = FText::Format(LOCTEXT("Class_NewDerivedBlueprintClassTooltip_CreateFrom", "Create a new Blueprint class based on {0}."), BaseClassName);
-		}
-		else
-		{
-			NewDerivedBlueprintClassToolTip = FText::Format(LOCTEXT("Class_NewDerivedBlueprintClassTooltip_InvalidClass", "Cannot create a new Blueprint class based on {0}."), BaseClassName);
-		}
-	}
-	else
-	{
-		NewDerivedCppClassLabel = LOCTEXT("Class_NewDerivedCppClassLabel_InvalidNumberOfBases", "New C++ class derived from...");
-		NewDerivedCppClassToolTip = LOCTEXT("Class_NewDerivedCppClassTooltip_InvalidNumberOfBases", "Can only create a derived C++ class when there is a single base class selected.");
-
-		NewDerivedBlueprintClassLabel = LOCTEXT("Class_NewDerivedBlueprintClassLabel_InvalidNumberOfBases", "New Blueprint class based on...");
-		NewDerivedBlueprintClassToolTip = LOCTEXT("Class_NewDerivedBlueprintClassTooltip_InvalidNumberOfBases", "Can only create a Blueprint class when there is a single base class selected.");
-	}
-
-	Section.AddMenuEntry(
-		"NewDerivedCppClass",
-		NewDerivedCppClassLabel,
-		NewDerivedCppClassToolTip,
-		FSlateIcon(FAppStyle::GetAppStyleSetName(), "Icons.C++"),
-		FUIAction(
-			FExecuteAction::CreateLambda(CreateCreateDerivedCppClass),
-			FCanExecuteAction::CreateLambda(CanCreateDerivedCppClass),
-			FIsActionChecked(),
-			FIsActionButtonVisible::CreateLambda(IsCPPAllowed)
-		)
-	);
-
-	Section.AddMenuEntry(
-		"NewDerivedBlueprintClass",
-		NewDerivedBlueprintClassLabel,
-		NewDerivedBlueprintClassToolTip,
-		FSlateIcon(FAppStyle::GetAppStyleSetName(), "Icons.Blueprint"),
-		FUIAction(
-			FExecuteAction::CreateLambda(CreateCreateDerivedBlueprintClass),
-			FCanExecuteAction::CreateLambda(CanCreateDerivedBlueprintClass)
-			)
-		);
-}
-
-UThumbnailInfo* FAssetTypeActions_Class::GetThumbnailInfo(UObject* Asset) const
+UThumbnailInfo* UAssetDefinition_Class::LoadThumbnailInfo(const FAssetData& InAssetData) const
 {
 	// todo: jdale - CLASS - We need to generate and store proper thumbnail info for classes so that we can store their custom render transforms
 	// This can't be stored in the UClass instance (like we do for Blueprints), so we'll need another place to store it
@@ -139,44 +26,176 @@ UThumbnailInfo* FAssetTypeActions_Class::GetThumbnailInfo(UObject* Asset) const
 	return nullptr;
 }
 
-void FAssetTypeActions_Class::OpenAssetEditor( const TArray<UObject*>& InObjects, TSharedPtr<class IToolkitHost> EditWithinLevelEditor )
+EAssetCommandResult UAssetDefinition_Class::OpenAssets(const FAssetOpenArgs& OpenArgs) const
 {
-	for(UObject* Object : InObjects)
+	TArray<FString> FilesToOpen;
+	for (UClass* Class : OpenArgs.LoadObjects<UClass>())
 	{
-		UClass* const Class = Cast<UClass>(Object);
-		if(Class)
+		FString ClassHeaderPath;
+		if(FSourceCodeNavigation::FindClassHeaderPath(Class, ClassHeaderPath))
 		{
-			TArray<FString> FilesToOpen;
-
-			FString ClassHeaderPath;
-			if(FSourceCodeNavigation::FindClassHeaderPath(Class, ClassHeaderPath))
-			{
-				const FString AbsoluteHeaderPath = IFileManager::Get().ConvertToAbsolutePathForExternalAppForRead(*ClassHeaderPath);
-				FilesToOpen.Add(AbsoluteHeaderPath);
-			}
-
-			FString ClassSourcePath;
-			if(FSourceCodeNavigation::FindClassSourcePath(Class, ClassSourcePath))
-			{
-				const FString AbsoluteSourcePath = IFileManager::Get().ConvertToAbsolutePathForExternalAppForRead(*ClassSourcePath);
-				FilesToOpen.Add(AbsoluteSourcePath);
-			}
+			const FString AbsoluteHeaderPath = IFileManager::Get().ConvertToAbsolutePathForExternalAppForRead(*ClassHeaderPath);
+			FilesToOpen.Add(AbsoluteHeaderPath);
+		}
 
-			FSourceCodeNavigation::OpenSourceFiles(FilesToOpen);
+		FString ClassSourcePath;
+		if(FSourceCodeNavigation::FindClassSourcePath(Class, ClassSourcePath))
+		{
+			const FString AbsoluteSourcePath = IFileManager::Get().ConvertToAbsolutePathForExternalAppForRead(*ClassSourcePath);
+			FilesToOpen.Add(AbsoluteSourcePath);
 		}
 	}
+
+	if (FilesToOpen.Num() > 0)
+	{
+		FSourceCodeNavigation::OpenSourceFiles(FilesToOpen);
+	}
+
+	return EAssetCommandResult::Handled;
 }
 
-TWeakPtr<IClassTypeActions> FAssetTypeActions_Class::GetClassTypeActions(const FAssetData& AssetData) const
+TWeakPtr<IClassTypeActions> UAssetDefinition_Class::GetClassTypeActions(const FAssetData& AssetData) const
 {
 	UClass* Class = FindObject<UClass>(nullptr, *AssetData.GetObjectPathString());
-	if(Class)
+	if (Class)
 	{
-		FAssetToolsModule& AssetToolsModule = FAssetToolsModule::GetModule();
-		return AssetToolsModule.Get().GetClassTypeActionsForClass(Class);
+		return IAssetTools::Get().GetClassTypeActionsForClass(Class);
 	}
 
 	return nullptr;
 }
 
+
+// Menu Extensions
+//--------------------------------------------------------------------
+
+namespace MenuExtension_Class
+{
+	static FDelayedAutoRegisterHelper DelayedAutoRegister(EDelayedRegisterRunPhase::EndOfEngineInit, []{ 
+		UToolMenus::RegisterStartupCallback(FSimpleMulticastDelegate::FDelegate::CreateLambda([]()
+		{
+			FToolMenuOwnerScoped OwnerScoped(UE_MODULE_NAME);
+			UToolMenu* Menu = UE::ContentBrowser::ExtendToolMenu_AssetContextMenu(UClass::StaticClass());
+		
+			FToolMenuSection& Section = Menu->FindOrAddSection("GetAssetActions");
+			Section.AddDynamicEntry(NAME_None, FNewToolMenuSectionDelegate::CreateLambda([](FToolMenuSection& InSection)
+			{
+				if (const UContentBrowserAssetContextMenuContext* CBContext = UContentBrowserAssetContextMenuContext::FindContextWithAssets(InSection))
+				{
+					UClass *const BaseClass = (CBContext->SelectedAssets.Num() == 1) ? Cast<UClass>(CBContext->SelectedAssets[0].GetAsset()) : nullptr;
+
+					// Only allow the New class option if we have a base class that we can actually derive from in one of our project modules
+					FGameProjectGenerationModule& GameProjectGenerationModule = FGameProjectGenerationModule::Get();
+					TArray<FModuleContextInfo> ProjectModules = GameProjectGenerationModule.GetCurrentProjectModules();
+					const bool bIsValidBaseCppClass = BaseClass && GameProjectGenerationModule.IsValidBaseClassForCreation(BaseClass, ProjectModules);
+					const bool bIsValidBaseBlueprintClass = BaseClass && FKismetEditorUtilities::CanCreateBlueprintOfClass(BaseClass);
+
+					auto CreateCreateDerivedCppClass = [BaseClass](const FToolMenuContext& InContext)
+					{
+						// Work out where the header file for the current class is, as we'll use that path as the default for the new class
+						FString BaseClassPath;
+						if(FSourceCodeNavigation::FindClassHeaderPath(BaseClass, BaseClassPath))
+						{
+							// Strip off the actual filename as we only need the path
+							BaseClassPath = FPaths::GetPath(BaseClassPath);
+						}
+
+						FGameProjectGenerationModule::Get().OpenAddCodeToProjectDialog(
+							FAddToProjectConfig()
+							.ParentClass(BaseClass)
+							.InitialPath(BaseClassPath)
+							.ParentWindow(FGlobalTabmanager::Get()->GetRootWindow())
+						);
+					};
+
+					auto IsCPPAllowed = []()
+					{
+						return ensure(GUnrealEd) && GUnrealEd->GetUnrealEdOptions()->IsCPPAllowed();
+					};
+
+					auto CanCreateDerivedCppClass = [bIsValidBaseCppClass, IsCPPAllowed](const FToolMenuContext& InContext) -> bool
+					{
+						return bIsValidBaseCppClass && IsCPPAllowed();
+					};
+
+					auto CreateCreateDerivedBlueprintClass = [BaseClass](const FToolMenuContext& InContext)
+					{
+						FGameProjectGenerationModule::Get().OpenAddBlueprintToProjectDialog(
+							FAddToProjectConfig()
+							.ParentClass(BaseClass)
+							.ParentWindow(FGlobalTabmanager::Get()->GetRootWindow())
+						);
+					};
+
+					auto CanCreateDerivedBlueprintClass = [bIsValidBaseBlueprintClass](const FToolMenuContext& InContext) -> bool
+					{
+						return bIsValidBaseBlueprintClass;
+					};
+
+					FText NewDerivedCppClassLabel;
+					FText NewDerivedCppClassToolTip;
+					FText NewDerivedBlueprintClassLabel;
+					FText NewDerivedBlueprintClassToolTip;
+					if (CBContext->SelectedAssets.Num() == 1)
+					{
+						check(BaseClass);
+
+						const FText BaseClassName = FText::FromName(BaseClass->GetFName());
+
+						NewDerivedCppClassLabel = FText::Format(LOCTEXT("Class_NewDerivedCppClassLabel_CreateFrom", "Create C++ class derived from {0}"), BaseClassName);
+						if(bIsValidBaseCppClass)
+						{
+							NewDerivedCppClassToolTip = FText::Format(LOCTEXT("Class_NewDerivedCppClassTooltip_CreateFrom", "Create a new C++ class deriving from {0}."), BaseClassName);
+						}
+						else
+						{
+							NewDerivedCppClassToolTip = FText::Format(LOCTEXT("Class_NewDerivedCppClassTooltip_InvalidClass", "Cannot create a new C++ class deriving from {0}."), BaseClassName);
+						}
+
+						NewDerivedBlueprintClassLabel = FText::Format(LOCTEXT("Class_NewDerivedBlueprintClassLabel_CreateFrom", "Create Blueprint class based on {0}"), BaseClassName);
+						if(bIsValidBaseBlueprintClass)
+						{
+							NewDerivedBlueprintClassToolTip = FText::Format(LOCTEXT("Class_NewDerivedBlueprintClassTooltip_CreateFrom", "Create a new Blueprint class based on {0}."), BaseClassName);
+						}
+						else
+						{
+							NewDerivedBlueprintClassToolTip = FText::Format(LOCTEXT("Class_NewDerivedBlueprintClassTooltip_InvalidClass", "Cannot create a new Blueprint class based on {0}."), BaseClassName);
+						}
+					}
+					else
+					{
+						NewDerivedCppClassLabel = LOCTEXT("Class_NewDerivedCppClassLabel_InvalidNumberOfBases", "New C++ class derived from...");
+						NewDerivedCppClassToolTip = LOCTEXT("Class_NewDerivedCppClassTooltip_InvalidNumberOfBases", "Can only create a derived C++ class when there is a single base class selected.");
+
+						NewDerivedBlueprintClassLabel = LOCTEXT("Class_NewDerivedBlueprintClassLabel_InvalidNumberOfBases", "New Blueprint class based on...");
+						NewDerivedBlueprintClassToolTip = LOCTEXT("Class_NewDerivedBlueprintClassTooltip_InvalidNumberOfBases", "Can only create a Blueprint class when there is a single base class selected.");
+					}
+
+					if (IsCPPAllowed())
+					{
+						const TAttribute<FText> Label = NewDerivedCppClassLabel;
+						const TAttribute<FText> ToolTip = NewDerivedCppClassToolTip;
+						const FSlateIcon Icon = FSlateIcon(FAppStyle::GetAppStyleSetName(), "Icons.C++");
+
+						FToolUIAction UIAction;
+						UIAction.ExecuteAction = FToolMenuExecuteAction::CreateLambda(CreateCreateDerivedCppClass);
+						UIAction.CanExecuteAction = FToolMenuCanExecuteAction::CreateLambda(CanCreateDerivedCppClass);
+						InSection.AddMenuEntry("NewDerivedCppClass", Label, ToolTip, Icon, UIAction);
+					}
+					{
+                    	const TAttribute<FText> Label = NewDerivedBlueprintClassLabel;
+                    	const TAttribute<FText> ToolTip = NewDerivedBlueprintClassToolTip;
+                    	const FSlateIcon Icon = FSlateIcon(FAppStyle::GetAppStyleSetName(), "Icons.C++");
+
+                    	FToolUIAction UIAction;
+                    	UIAction.ExecuteAction = FToolMenuExecuteAction::CreateLambda(CreateCreateDerivedBlueprintClass);
+			            UIAction.CanExecuteAction = FToolMenuCanExecuteAction::CreateLambda(CanCreateDerivedBlueprintClass);
+                    	InSection.AddMenuEntry("NewDerivedBlueprintClass", Label, ToolTip, Icon, UIAction);
+                    }
+				}
+			}));
+		}));
+	});
+}
+
 #undef LOCTEXT_NAMESPACE
diff --git a/Engine/Plugins/Editor/EngineAssetDefinitions/Source/Private/Script/AssetDefinition_Class.h b/Engine/Plugins/Editor/EngineAssetDefinitions/Source/Private/Script/AssetDefinition_Class.h
--- a/Engine/Plugins/Editor/EngineAssetDefinitions/Source/Private/Script/AssetDefinition_Class.h
+++ b/Engine/Plugins/Editor/EngineAssetDefinitions/Source/Private/Script/AssetDefinition_Class.h
@@ -2,29 +2,36 @@
 
 #pragma once
 
-#include "CoreMinimal.h"
-#include "UObject/Class.h"
-#include "Toolkits/IToolkitHost.h"
-#include "AssetTypeActions/AssetTypeActions_ClassTypeBase.h"
+#include "Script/AssetDefinition_ClassTypeBase.h"
 
+#include "AssetDefinition_Class.generated.h"
+
+struct FToolMenuContext;
 struct FAssetData;
-class IClassTypeActions;
+class UFactory;
 
-class FAssetTypeActions_Class : public FAssetTypeActions_ClassTypeBase
+UCLASS()
+class UAssetDefinition_Class : public UAssetDefinition_ClassTypeBase
 {
+	GENERATED_BODY()
+
 public:
-	// IAssetTypeActions Implementation
-	virtual FText GetName() const override { return NSLOCTEXT("AssetTypeActions", "AssetTypeActions_Class", "C++ Class"); }
-	virtual FColor GetTypeColor() const override { return FColor(255, 255, 255); }
-	virtual UClass* GetSupportedClass() const override { return UClass::StaticClass(); }
-	virtual uint32 GetCategories() override { return EAssetTypeCategories::Basic; }
+	// UAssetDefinition Implementation
+	virtual FText GetAssetDisplayName() const override { return NSLOCTEXT("AssetTypeActions", "AssetTypeActions_Class", "C++ Class"); }
+	virtual FLinearColor GetAssetColor() const override { return FLinearColor(FColor(255, 255, 255)); }
+	virtual TSoftClassPtr<UObject> GetAssetClass() const override { return UClass::StaticClass(); }
+	virtual TConstArrayView<FAssetCategoryPath> GetAssetCategories() const override
+	{
+		static const auto Categories = { EAssetCategoryPaths::Basic };
+		return Categories;
+	}
 	
-	virtual void GetActions(const TArray<UObject*>& InObjects, struct FToolMenuSection& Section) override;
-
-	virtual class UThumbnailInfo* GetThumbnailInfo(UObject* Asset) const override;
+	virtual UThumbnailInfo* LoadThumbnailInfo(const FAssetData& InAssetData) const override;
+	virtual EAssetCommandResult OpenAssets(const FAssetOpenArgs& OpenArgs) const override;
+	// UAssetDefinition End
 
-	virtual void OpenAssetEditor( const TArray<UObject*>& InObjects, TSharedPtr<class IToolkitHost> EditWithinLevelEditor = TSharedPtr<IToolkitHost>() ) override;
-
-	// FAssetTypeActions_ClassTypeBase Implementation
+	// UAssetDefinition_ClassTypeBase Implementation
 	virtual TWeakPtr<IClassTypeActions> GetClassTypeActions(const FAssetData& AssetData) const override;
+	// UAssetDefinition_ClassTypeBase End
 };
+
diff --git a/Engine/Source/Developer/AssetTools/Private/AssetTools.cpp b/Engine/Source/Developer/AssetTools/Private/AssetTools.cpp
--- a/Engine/Source/Developer/AssetTools/Private/AssetTools.cpp
+++ b/Engine/Source/Developer/AssetTools/Private/AssetTools.cpp
@@ -354,8 +354,7 @@ public:
 	/** Returns additional tooltip information for the specified asset, if it has any (otherwise return the null widget) */
 	virtual FText GetAssetDescription(const FAssetData& AssetData) const override
 	{
-		ensureMsgf(false, TEXT("This code path is not expected to be called any more.  Expected UAssetDefinitionRegistry::Get()->GetAssetDefinitionForAsset(YourAsset)->GetAssetDescription to be called directly."));
-        return FText::GetEmpty();
+		return AssetDefinitionPtr.Get()->GetAssetDescription(AssetData);
 	}
 
 	/** Collects the resolved source paths for the imported assets */
