Hide "Convert to Different DataAsset Type" menu if none of the selected asset can be edited

diff --git a/Engine/Plugins/Editor/EngineAssetDefinitions/Source/EngineAssetDefinitions.Build.cs b/Engine/Plugins/Editor/EngineAssetDefinitions/Source/EngineAssetDefinitions.Build.cs
--- a/Engine/Plugins/Editor/EngineAssetDefinitions/Source/EngineAssetDefinitions.Build.cs
+++ b/Engine/Plugins/Editor/EngineAssetDefinitions/Source/EngineAssetDefinitions.Build.cs
@@ -17,6 +17,7 @@ public class EngineAssetDefinitions : ModuleRules
 			new string[]
 			{
 				"ContentBrowser",
+				"ContentBrowserData",
 				"AssetDefinition",
 				"UnrealEd",
 				"ToolMenus",
diff --git a/Engine/Plugins/Editor/EngineAssetDefinitions/Source/Private/AssetDefinition_DataAsset.cpp b/Engine/Plugins/Editor/EngineAssetDefinitions/Source/Private/AssetDefinition_DataAsset.cpp
--- a/Engine/Plugins/Editor/EngineAssetDefinitions/Source/Private/AssetDefinition_DataAsset.cpp
+++ b/Engine/Plugins/Editor/EngineAssetDefinitions/Source/Private/AssetDefinition_DataAsset.cpp
@@ -42,12 +42,25 @@ namespace MenuExtension_DataAsset
 		UClass* BaseClass;
 	};
 
-	void ExecuteChangeDataAssetClass(const FToolMenuContext& MenuContext)
+	bool IsChangeDataAssetClassVisible(const FToolMenuContext& MenuContext)
 	{
 		if (const UContentBrowserAssetContextMenuContext* Context = UContentBrowserAssetContextMenuContext::FindContextWithAssets(MenuContext))
 		{
-			TArray<UDataAsset*> DataAssets = Context->LoadSelectedObjects<UDataAsset>();
+			for (const FContentBrowserItem& SelectedItem : Context->GetSelectedItems())
+			{
+				if (SelectedItem.CanEdit())
+				{
+					return true;
+				}
+			}
+		}
+		return false;
+	}
 
+	void ExecuteChangeDataAssetClass(const FToolMenuContext& MenuContext)
+	{
+		if (const UContentBrowserAssetContextMenuContext* Context = UContentBrowserAssetContextMenuContext::FindContextWithAssets(MenuContext))
+		{
 			const FText TitleText = LOCTEXT("DataAsset_PickNewDataAssetClass", "Pick New DataAsset Class");
 			FClassViewerInitializationOptions Options;
 			Options.ClassFilters.Add(MakeShared<FNewNodeClassFilter>(UDataAsset::StaticClass()));
@@ -56,6 +69,25 @@ namespace MenuExtension_DataAsset
 
 			if (bPressedOk && OutNewDataAssetClass != nullptr)
 			{
+				TSet<FName> EditableAssets;
+				{
+					const TArray<FContentBrowserItem>& SelectedItems = Context->GetSelectedItems();
+					EditableAssets.Reserve(SelectedItems.Num());
+					for (const FContentBrowserItem& SelectedItem : Context->GetSelectedItems())
+					{
+						if (SelectedItem.CanEdit())
+						{
+							EditableAssets.Add(SelectedItem.GetInternalPath());
+						}
+					}
+				}
+				ensure(!EditableAssets.IsEmpty());
+
+				TArray<UDataAsset*> DataAssets = Context->LoadSelectedObjectsIf<UDataAsset>([&EditableAssets](const FAssetData& AssetData)
+				{
+					return EditableAssets.Contains(*AssetData.GetObjectPathString());
+				});
+
 				for (TWeakObjectPtr<UDataAsset> DataAssetPtr : DataAssets)
 				{
 					if (UDataAsset* OldDataAsset = DataAssetPtr.Get())
@@ -101,8 +133,11 @@ namespace MenuExtension_DataAsset
 					const TAttribute<FText> Label = LOCTEXT("DataAsset_ChangeClass", "Convert to Different DataAsset Type");
 					const TAttribute<FText> ToolTip = LOCTEXT("DataAsset_ChangeClassTip", "Change the class these Data Assets are subclassed from.");
 					const FSlateIcon Icon = FSlateIcon(FAppStyle::GetAppStyleSetName(), "ClassIcon.DataAsset");
-					const FToolMenuExecuteAction UIAction = FToolMenuExecuteAction::CreateStatic(&ExecuteChangeDataAssetClass);
-	    
+
+					FToolUIAction UIAction;
+					UIAction.ExecuteAction = FToolMenuExecuteAction::CreateStatic(&ExecuteChangeDataAssetClass);
+					UIAction.IsActionVisibleDelegate = FToolMenuIsActionButtonVisible::CreateStatic(&IsChangeDataAssetClassVisible);
+
 					InSection.AddMenuEntry("DataAsset_ChangeClass", Label, ToolTip, Icon, UIAction);
 				}
 			}));
