Made the default texture 2D contextual menu aware of the class permission system

diff --git a/Engine/Plugins/Editor/EngineAssetDefinitions/Source/Private/AssetDefinition_Texture2D.cpp b/Engine/Plugins/Editor/EngineAssetDefinitions/Source/Private/AssetDefinition_Texture2D.cpp
--- a/Engine/Plugins/Editor/EngineAssetDefinitions/Source/Private/AssetDefinition_Texture2D.cpp
+++ b/Engine/Plugins/Editor/EngineAssetDefinitions/Source/Private/AssetDefinition_Texture2D.cpp
@@ -5,6 +5,7 @@
 #include "AssetToolsModule.h"
 #include "ContentBrowserMenuContexts.h"
 #include "ToolMenus.h"
+#include "Misc/NamePermissionList.h"
 #include "Misc/PackageName.h"
 #include "Factories/SlateBrushAssetFactory.h"
 #include "Slate/SlateBrushAsset.h"
@@ -118,7 +119,9 @@ namespace MenuExtension_Texture2D
 			Section.AddDynamicEntry(NAME_None, FNewToolMenuSectionDelegate::CreateLambda([](FToolMenuSection& InSection)
 			{
 				const UContentBrowserAssetContextMenuContext* Context = UContentBrowserAssetContextMenuContext::FindContextWithAssets(InSection);
-				
+				TSharedRef<FPathPermissionList> ClassPathPermissionList = IAssetTools::Get().GetAssetClassPathPermissionList(EAssetClassAction::CreateAsset);
+
+				if (ClassPathPermissionList->PassesFilter(USlateBrushAsset::StaticClass()->GetClassPathName().ToString()))
 				{
 					const TAttribute<FText> Label = LOCTEXT("Texture2D_CreateSlateBrush", "Create Slate Brush");
 					const TAttribute<FText> ToolTip = LOCTEXT("Texture2D_CreateSlateBrushToolTip", "Creates a new slate brush using this texture.");
@@ -128,18 +131,21 @@ namespace MenuExtension_Texture2D
 					InSection.AddMenuEntry("Texture2D_CreateSlateBrush", Label, ToolTip, Icon, UIAction);
 				}
 
-				static const auto AllowTextureArrayAssetCreationVar = IConsoleManager::Get().FindTConsoleVariableDataInt(TEXT("r.AllowTexture2DArrayCreation"));
-				if (AllowTextureArrayAssetCreationVar->GetValueOnGameThread() != 0)
+				if (ClassPathPermissionList->PassesFilter(UTexture2DArray::StaticClass()->GetClassPathName().ToString()))
 				{
-					const TAttribute<FText> Label = LOCTEXT("Texture_Texture2DArray", "Create Texture Array");
-					const TAttribute<FText> ToolTip = LOCTEXT("Texture_CreateTexture2DArrayTooltip", "Creates a new texture array.");
-					const FSlateIcon Icon = FSlateIcon(FAppStyle::GetAppStyleSetName(), "ClassIcon.Texture2D");
-					const FToolMenuExecuteAction UIAction = FToolMenuExecuteAction::CreateStatic(&ExecuteCreateTextureArray);
-
-					InSection.AddMenuEntry("Texture_Texture2DArray", Label, ToolTip, Icon, UIAction);
+					static const auto AllowTextureArrayAssetCreationVar = IConsoleManager::Get().FindTConsoleVariableDataInt(TEXT("r.AllowTexture2DArrayCreation"));
+					if (AllowTextureArrayAssetCreationVar->GetValueOnGameThread() != 0)
+					{
+						const TAttribute<FText> Label = LOCTEXT("Texture_Texture2DArray", "Create Texture Array");
+						const TAttribute<FText> ToolTip = LOCTEXT("Texture_CreateTexture2DArrayTooltip", "Creates a new texture array.");
+						const FSlateIcon Icon = FSlateIcon(FAppStyle::GetAppStyleSetName(), "ClassIcon.Texture2D");
+						const FToolMenuExecuteAction UIAction = FToolMenuExecuteAction::CreateStatic(&ExecuteCreateTextureArray);
+
+						InSection.AddMenuEntry("Texture_Texture2DArray", Label, ToolTip, Icon, UIAction);
+					}
 				}
 
-				if (Context->SelectedAssets.Num() == 1)
+				if (Context->SelectedAssets.Num() == 1 && ClassPathPermissionList->PassesFilter(UVolumeTexture::StaticClass()->GetClassPathName().ToString()))
                 {
 					const TAttribute<FText> Label = LOCTEXT("Texture2D_CreateVolumeTexture", "Create Volume Texture");
 					const TAttribute<FText> ToolTip = LOCTEXT("Texture2D_CreateVolumeTextureToolTip", "Creates a new volume texture using this texture.");
