Converted SparseVolumeTextures from AssetTypeActions to AssetDefinition.

#rnx

diff --git a/Engine/Plugins/Editor/EngineAssetDefinitions/Source/Private/AssetDefinition_SparseVolumeTexture.cpp b/Engine/Plugins/Editor/EngineAssetDefinitions/Source/Private/AssetDefinition_SparseVolumeTexture.cpp
new file mode 100644
--- /dev/null
+++ b/Engine/Plugins/Editor/EngineAssetDefinitions/Source/Private/AssetDefinition_SparseVolumeTexture.cpp
@@ -0,0 +1,11 @@
+// Copyright Epic Games, Inc. All Rights Reserved.
+
+#include "AssetDefinition_SparseVolumeTexture.h"
+
+#define LOCTEXT_NAMESPACE "UAssetDefinition_SparseVolumeTexture"
+
+namespace MenuExtension_SparseVolumeTexture
+{
+}
+
+#undef LOCTEXT_NAMESPACE
diff --git a/Engine/Plugins/Editor/EngineAssetDefinitions/Source/Private/AssetDefinition_SparseVolumeTexture.h b/Engine/Plugins/Editor/EngineAssetDefinitions/Source/Private/AssetDefinition_SparseVolumeTexture.h
new file mode 100644
--- /dev/null
+++ b/Engine/Plugins/Editor/EngineAssetDefinitions/Source/Private/AssetDefinition_SparseVolumeTexture.h
@@ -0,0 +1,51 @@
+// Copyright Epic Games, Inc. All Rights Reserved.
+
+#pragma once
+
+#include "SparseVolumeTexture/SparseVolumeTexture.h"
+#include "AssetDefinitionDefault.h"
+
+#include "AssetDefinition_SparseVolumeTexture.generated.h"
+
+UCLASS()
+class UAssetDefinition_SparseVolumeTexture : public UAssetDefinitionDefault
+{
+	GENERATED_BODY()
+
+public:
+	// UAssetDefinition Begin
+	virtual FText GetAssetDisplayName() const override { return NSLOCTEXT("AssetTypeActions", "AssetTypeActions_SparseVolumeTexture", "BaseSparseVolumeTexture"); }
+	virtual FLinearColor GetAssetColor() const override { return FLinearColor(FColor(255, 128, 64)); }
+	virtual TSoftClassPtr<UObject> GetAssetClass() const override { return USparseVolumeTexture::StaticClass(); }
+	virtual TConstArrayView<FAssetCategoryPath> GetAssetCategories() const override
+	{
+		static const auto Categories = { EAssetCategoryPaths::Texture };
+		return Categories;
+	}
+	virtual bool CanImport() const override { return true; }
+	// UAssetDefinition End
+};
+
+UCLASS()
+class UAssetDefinition_StaticSparseVolumeTexture : public UAssetDefinition_SparseVolumeTexture
+{
+	GENERATED_BODY()
+
+public:
+	// UAssetDefinition Begin
+	virtual FText GetAssetDisplayName() const override { return NSLOCTEXT("AssetTypeActions", "AssetTypeActions_StaticSparseVolumeTexture", "StaticSparseVolumeTexture"); }
+	virtual TSoftClassPtr<UObject> GetAssetClass() const override { return UStaticSparseVolumeTexture::StaticClass(); }
+	// UAssetDefinition End
+};
+
+UCLASS()
+class UAssetDefinition_AnimatedSparseVolumeTexture : public UAssetDefinition_SparseVolumeTexture
+{
+	GENERATED_BODY()
+
+public:
+	// UAssetDefinition Begin
+	virtual FText GetAssetDisplayName() const override { return NSLOCTEXT("AssetTypeActions", "AssetTypeActions_AnimatedSparseVolumeTexture", "AnimatedSparseVolumeTexture"); }
+	virtual TSoftClassPtr<UObject> GetAssetClass() const override { return UAnimatedSparseVolumeTexture::StaticClass(); }
+	// UAssetDefinition End
+};
diff --git a/Engine/Source/Editor/SparseVolumeTexture/Private/AssetTypeActions_SparseVolumeTexture.cpp b/Engine/Source/Editor/SparseVolumeTexture/Private/AssetTypeActions_SparseVolumeTexture.cpp
deleted file mode 100644
diff --git a/Engine/Source/Editor/SparseVolumeTexture/Private/AssetTypeActions_SparseVolumeTexture.h b/Engine/Source/Editor/SparseVolumeTexture/Private/AssetTypeActions_SparseVolumeTexture.h
deleted file mode 100644
diff --git a/Engine/Source/Editor/SparseVolumeTexture/Private/SparseVolumeTextureModule.cpp b/Engine/Source/Editor/SparseVolumeTexture/Private/SparseVolumeTextureModule.cpp
--- a/Engine/Source/Editor/SparseVolumeTexture/Private/SparseVolumeTextureModule.cpp
+++ b/Engine/Source/Editor/SparseVolumeTexture/Private/SparseVolumeTextureModule.cpp
@@ -10,10 +10,6 @@ IMPLEMENT_MODULE(FSparseVolumeTextureModule, SparseVolumeTexture);
 
 void FSparseVolumeTextureModule::StartupModule()
 {
-	IAssetTools& AssetTools = FModuleManager::LoadModuleChecked<FAssetToolsModule>(TEXT("AssetTools")).Get();
-	AssetTypeActionsForSparseVolumeTexture = MakeShareable(new FAssetTypeActions_SparseVolumeTexture);
-	AssetTools.RegisterAssetTypeActions(AssetTypeActionsForSparseVolumeTexture.ToSharedRef());
-
 #if PLATFORM_WINDOWS
 	// Global registration of  the vdb types.
 	openvdb::initialize();
diff --git a/Engine/Source/Editor/SparseVolumeTexture/Public/SparseVolumeTextureModule.h b/Engine/Source/Editor/SparseVolumeTexture/Public/SparseVolumeTextureModule.h
--- a/Engine/Source/Editor/SparseVolumeTexture/Public/SparseVolumeTextureModule.h
+++ b/Engine/Source/Editor/SparseVolumeTexture/Public/SparseVolumeTextureModule.h
@@ -7,8 +7,6 @@
 #include "IAssetTools.h"
 #include "AssetRegistry/AssetRegistryModule.h"
 
-#include "AssetTypeActions_SparseVolumeTexture.h"
-
 /**
  * A sparse volume texture
  */
@@ -34,6 +32,4 @@ private:
 	 * Called before the module is unloaded, right before the module object is destroyed.
 	 */
 	virtual void ShutdownModule();
-
-	TSharedPtr<FAssetTypeActions_SparseVolumeTexture> AssetTypeActionsForSparseVolumeTexture;
 };
