Fix being unable to make a child animation blueprint via the content browser

Switched submenu that required loading the skeleton over to use a dialog so the menu doesnt have to load any assets

diff --git a/Engine/Plugins/Editor/EngineAssetDefinitions/Source/EngineAssetDefinitions.Build.cs b/Engine/Plugins/Editor/EngineAssetDefinitions/Source/EngineAssetDefinitions.Build.cs
--- a/Engine/Plugins/Editor/EngineAssetDefinitions/Source/EngineAssetDefinitions.Build.cs
+++ b/Engine/Plugins/Editor/EngineAssetDefinitions/Source/EngineAssetDefinitions.Build.cs
@@ -36,7 +36,8 @@ public class EngineAssetDefinitions : ModuleRules
 				"PhysicsCore",
 				"Kismet",
 				"GameProjectGeneration",
-				"AnimationBlueprintEditor"
+				"AnimationBlueprintEditor",
+				"ToolWidgets"
 			}
 		);
 		
diff --git a/Engine/Plugins/Editor/EngineAssetDefinitions/Source/Private/Animation/AssetDefinition_AnimBlueprint.cpp b/Engine/Plugins/Editor/EngineAssetDefinitions/Source/Private/Animation/AssetDefinition_AnimBlueprint.cpp
--- a/Engine/Plugins/Editor/EngineAssetDefinitions/Source/Private/Animation/AssetDefinition_AnimBlueprint.cpp
+++ b/Engine/Plugins/Editor/EngineAssetDefinitions/Source/Private/Animation/AssetDefinition_AnimBlueprint.cpp
@@ -20,6 +20,8 @@
 #include "ToolMenu.h"
 #include "ToolMenuEntry.h"
 #include "ToolMenuSection.h"
+#include "SPrimaryButton.h"
+#include "SSimpleButton.h"
 
 #define LOCTEXT_NAMESPACE "AssetTypeActions"
 
@@ -244,15 +246,112 @@ namespace MenuExtension_AnimBlueprint
 		}
 	}
 
-	void BuildNewSkeletonChildBlueprintMenu(UToolMenu* Menu, const FAssetData InAnimBlueprintAsset)
+	void ExecuteNewSkeletonChildBlueprint(const FToolMenuContext& InContext, const FAssetData InAnimBlueprintAsset)
 	{
-		auto HandleAssetSelected = [InAnimBlueprintAsset](const FAssetData& InSelectedSkeletonAsset)
+		TArray<FSoftObjectPath> CompatibleSkeletonPaths;
+		TSoftObjectPtr<USkeleton> AnimBlueprintSkeletonPtr = UE::AnimBlueprint::GetAnimBlueprintTargetSkeleton(InAnimBlueprintAsset);
+		if (USkeleton* AnimBlueprintSkeleton = AnimBlueprintSkeletonPtr.LoadSynchronous())
+		{
+			TArray<FAssetData> CompatibleSkeletonAssets;
+			AnimBlueprintSkeleton->GetCompatibleSkeletonAssets(CompatibleSkeletonAssets);
+
+			Algo::Transform(CompatibleSkeletonAssets, CompatibleSkeletonPaths, [](const FAssetData& InAsset) { return InAsset.GetSoftObjectPath(); });
+		}
+
+		FAssetData SelectedSkeletonAsset;
+		auto HandleAssetSelected = [&SelectedSkeletonAsset](const FAssetData& InSelectedSkeletonAsset)
+		{
+			SelectedSkeletonAsset = InSelectedSkeletonAsset;
+		};
+		
+		FAssetPickerConfig AssetPickerConfig;
+		AssetPickerConfig.OnAssetEnterPressed = FOnAssetEnterPressed::CreateLambda([HandleAssetSelected](const TArray<FAssetData>& SelectedAssetData)
+		{
+			if (SelectedAssetData.Num() == 1)
+			{
+				HandleAssetSelected(SelectedAssetData[0]);
+			}
+		});
+		AssetPickerConfig.OnAssetSelected = FOnAssetSelected::CreateLambda(HandleAssetSelected);
+		AssetPickerConfig.SelectionMode = ESelectionMode::Single;
+		AssetPickerConfig.bAllowNullSelection = false;
+		AssetPickerConfig.InitialAssetViewType = EAssetViewType::List;
+		AssetPickerConfig.Filter.SoftObjectPaths = MoveTemp(CompatibleSkeletonPaths);
+		AssetPickerConfig.Filter.ClassPaths.Add(USkeleton::StaticClass()->GetClassPathName());
+
+		FContentBrowserModule& ContentBrowserModule = FModuleManager::Get().LoadModuleChecked<FContentBrowserModule>(TEXT("ContentBrowser"));
+
+		bool bClickedOk = false;
+		TSharedPtr<SWindow> Dialog;
+		Dialog = SNew(SWindow)
+				.Title(FText(LOCTEXT("CompatibleSkeletonDialog", "Compatible Skeleton for Animation Blueprint")))
+				.SizingRule(ESizingRule::Autosized);
+
+		Dialog->SetContent(
+			SNew(SVerticalBox)
+			+SVerticalBox::Slot()
+			.AutoHeight()
+			.Padding(5.0f)
+			[
+				SNew(SBox)
+				.WidthOverride(300.0f)
+				.HeightOverride(400.0f)
+				[
+					ContentBrowserModule.Get().CreateAssetPicker(AssetPickerConfig)
+				]
+			]
+			+SVerticalBox::Slot()
+			.AutoHeight()
+			.HAlign(HAlign_Right)
+			.Padding(5.0f)
+			[
+				SNew(SUniformGridPanel)
+				.SlotPadding(FAppStyle::GetMargin("StandardDialog.SlotPadding"))
+				.MinDesiredSlotWidth(FAppStyle::GetFloat("StandardDialog.MinDesiredSlotWidth"))
+				.MinDesiredSlotHeight(FAppStyle::GetFloat("StandardDialog.MinDesiredSlotHeight"))
+				+SUniformGridPanel::Slot(0,0)
+				[
+					SNew(SPrimaryButton)
+					.OnClicked_Lambda([&bClickedOk, WeakDialog = TWeakPtr<SWindow>(Dialog)]()
+					{
+						bClickedOk = true;
+						if(TSharedPtr<SWindow> PinnedDialog = WeakDialog.Pin())
+						{
+							PinnedDialog->RequestDestroyWindow();
+						}
+						
+						return FReply::Handled();
+					})
+					.IsEnabled_Lambda([&SelectedSkeletonAsset]()
+					{
+						return SelectedSkeletonAsset.IsValid();
+					})
+					.Text(LOCTEXT("CreateChildButton", "Create Child"))
+				]
+				+SUniformGridPanel::Slot(1,0)
+				[
+					SNew(SButton)
+					.HAlign(HAlign_Center)
+					.OnClicked_Lambda([WeakDialog = TWeakPtr<SWindow>(Dialog)]()
+					{
+						if(TSharedPtr<SWindow> PinnedDialog = WeakDialog.Pin())
+						{
+							PinnedDialog->RequestDestroyWindow();
+						}
+
+						return FReply::Handled();
+					})
+					.Text(LOCTEXT("CancelButton", "Cancel"))
+				]
+			]);
+
+		FSlateApplication::Get().AddModalWindow(Dialog.ToSharedRef(), FGlobalTabmanager::Get()->GetRootWindow());
+
+		if(bClickedOk && SelectedSkeletonAsset.IsValid())	// Clicked OK
 		{
-			FSlateApplication::Get().DismissAllMenus();
-			
 			if (UAnimBlueprint* TargetParentBP = Cast<UAnimBlueprint>(InAnimBlueprintAsset.GetAsset()))
 			{
-				USkeleton* TargetSkeleton = CastChecked<USkeleton>(InSelectedSkeletonAsset.GetAsset());
+				USkeleton* TargetSkeleton = CastChecked<USkeleton>(SelectedSkeletonAsset.GetAsset());
 				UClass* TargetParentClass = TargetParentBP->GeneratedClass;
 
 				if (!FKismetEditorUtilities::CanCreateBlueprintOfClass(TargetParentClass))
@@ -271,51 +370,11 @@ namespace MenuExtension_AnimBlueprint
 				AnimBlueprintFactory->TargetSkeleton = TargetSkeleton;
 				AnimBlueprintFactory->bTemplate = false;
 
-				FContentBrowserModule& ContentBrowserModule = FModuleManager::LoadModuleChecked<FContentBrowserModule>("ContentBrowser");
 				ContentBrowserModule.Get().CreateNewAsset(Name, PackagePath, TargetParentBP->GetClass(), AnimBlueprintFactory);
 			}
-		};
-
-		TSoftObjectPtr<USkeleton> AnimBlueprintSkeletonPtr = UE::AnimBlueprint::GetAnimBlueprintTargetSkeleton(InAnimBlueprintAsset);
-		if (USkeleton* AnimBlueprintSkeleton = AnimBlueprintSkeletonPtr.Get())
-		{
-			TArray<FAssetData> CompatibleSkeletonAssets;
-			AnimBlueprintSkeleton->GetCompatibleSkeletonAssets(CompatibleSkeletonAssets);
-
-			TArray<FSoftObjectPath> CompatibleSkeletonPaths;
-			Algo::Transform(CompatibleSkeletonAssets, CompatibleSkeletonPaths, [](const FAssetData& InAsset) { return InAsset.GetSoftObjectPath(); });
-		
-
-			FAssetPickerConfig AssetPickerConfig;
-			AssetPickerConfig.OnAssetEnterPressed = FOnAssetEnterPressed::CreateLambda([HandleAssetSelected](const TArray<FAssetData>& SelectedAssetData)
-			{
-				if (SelectedAssetData.Num() == 1)
-				{
-					HandleAssetSelected(SelectedAssetData[0]);
-				}
-			});
-			AssetPickerConfig.OnAssetSelected = FOnAssetSelected::CreateLambda(HandleAssetSelected);
-			AssetPickerConfig.bAllowNullSelection = false;
-			AssetPickerConfig.InitialAssetViewType = EAssetViewType::List;
-			AssetPickerConfig.Filter.SoftObjectPaths = MoveTemp(CompatibleSkeletonPaths);
-		
-			FContentBrowserModule& ContentBrowserModule = FModuleManager::Get().LoadModuleChecked<FContentBrowserModule>(TEXT("ContentBrowser"));
-
-			FToolMenuSection& InSection = Menu->AddSection("CompatibleSkeletonMenu", LOCTEXT("CompatibleSkeletonHeader", "Compatible Skeletons"));
-			InSection.AddEntry(
-				FToolMenuEntry::InitWidget("CompatibleSkeletonPicker",
-					SNew(SBox)
-					.WidthOverride(300.0f)
-					.HeightOverride(300.0f)
-					[
-						ContentBrowserModule.Get().CreateAssetPicker(AssetPickerConfig)
-					],
-					FText::GetEmpty()
-				)
-			);
 		}
 	}
-		
+
 	static FDelayedAutoRegisterHelper DelayedAutoRegister(EDelayedRegisterRunPhase::EndOfEngineInit, []{ 
 		UToolMenus::RegisterStartupCallback(FSimpleMulticastDelegate::FDelegate::CreateLambda([]()
 		{
@@ -335,18 +394,17 @@ namespace MenuExtension_AnimBlueprint
 							(
 								(UE::AnimBlueprint::GetAnimBlueprintTargetSkeleton(*SelectedAnimBlueprintPtr).IsNull() && UE::AnimBlueprint::IsAnimBlueprintTemplate(*SelectedAnimBlueprintPtr))
 								||
-								(!UE::AnimBlueprint::GetAnimBlueprintTargetSkeleton(*SelectedAnimBlueprintPtr).IsNull() /* && AnimBlueprint->TargetSkeleton->GetCompatibleSkeletons().Num() > 0 */)
+								(!UE::AnimBlueprint::GetAnimBlueprintTargetSkeleton(*SelectedAnimBlueprintPtr).IsNull())
 							)
 						)
 						{
-							InSection.AddSubMenu(
-								"AnimBlueprint_NewSkeletonChildBlueprint",
-								LOCTEXT("AnimBlueprint_NewSkeletonChildBlueprint", "Create Child Anim Blueprint with Skeleton"),
-								LOCTEXT("AnimBlueprint_NewSkeletonChildBlueprint_Tooltip", "Create a child Anim Blueprint that uses a different compatible skeleton"),
-								FNewToolMenuDelegate::CreateStatic(&BuildNewSkeletonChildBlueprintMenu, *SelectedAnimBlueprintPtr),
-								false,
-								FSlateIcon(FAppStyle::GetAppStyleSetName(), "Icons.Blueprint")
-							);
+							const TAttribute<FText> Label = LOCTEXT("AnimBlueprint_NewSkeletonChildBlueprint", "Create Child Anim Blueprint with Skeleton");
+							const TAttribute<FText> ToolTip = LOCTEXT("AnimBlueprint_NewSkeletonChildBlueprint_Tooltip", "Create a child Anim Blueprint that uses a different compatible skeleton.");
+							const FSlateIcon Icon = FSlateIcon(FAppStyle::GetAppStyleSetName(), "Icons.Blueprint");
+
+							FToolUIAction UIAction;
+							UIAction.ExecuteAction = FToolMenuExecuteAction::CreateStatic(&ExecuteNewSkeletonChildBlueprint, *SelectedAnimBlueprintPtr);
+							InSection.AddMenuEntry("AnimBlueprint_NewSkeletonChildBlueprint", Label, ToolTip, Icon, UIAction);
 						}
 					}
 					
