Disable the "Create Blueprint class" based on non bp script

#rnx

diff --git a/Engine/Plugins/Editor/EngineAssetDefinitions/Source/EngineAssetDefinitions.Build.cs b/Engine/Plugins/Editor/EngineAssetDefinitions/Source/EngineAssetDefinitions.Build.cs
--- a/Engine/Plugins/Editor/EngineAssetDefinitions/Source/EngineAssetDefinitions.Build.cs
+++ b/Engine/Plugins/Editor/EngineAssetDefinitions/Source/EngineAssetDefinitions.Build.cs
@@ -40,7 +40,15 @@ public class EngineAssetDefinitions : ModuleRules
 				"ToolWidgets"
 			}
 		);
-		
+
+		if (Target.bUseVerse)
+		{
+			PrivateDependencyModuleNames.AddRange(new string[]
+			{
+				"VerseNative",
+			});
+		}
+
 		PrivateIncludePathModuleNames.AddRange(
 			new string[] {
 				"VirtualTexturingEditor",
diff --git a/Engine/Plugins/Editor/EngineAssetDefinitions/Source/Private/Script/AssetDefinition_Class.cpp b/Engine/Plugins/Editor/EngineAssetDefinitions/Source/Private/Script/AssetDefinition_Class.cpp
--- a/Engine/Plugins/Editor/EngineAssetDefinitions/Source/Private/Script/AssetDefinition_Class.cpp
+++ b/Engine/Plugins/Editor/EngineAssetDefinitions/Source/Private/Script/AssetDefinition_Class.cpp
@@ -17,8 +17,11 @@
 #include "ToolMenuSection.h"
 #include "UnrealEdGlobals.h"
 
-#define LOCTEXT_NAMESPACE "AssetTypeActions"
+#if WITH_VERSE
+#include "VerseClassBase.h"
+#endif
 
+#define LOCTEXT_NAMESPACE "AssetTypeActions"
 
 UThumbnailInfo* UAssetDefinition_Class::LoadThumbnailInfo(const FAssetData& InAssetData) const
 {
@@ -92,6 +95,14 @@ namespace MenuExtension_Class
 					const bool bIsValidBaseCppClass = BaseClass && GameProjectGenerationModule.IsValidBaseClassForCreation(BaseClass, ProjectModules);
 					const bool bIsValidBaseBlueprintClass = BaseClass && FKismetEditorUtilities::CanCreateBlueprintOfClass(BaseClass);
 
+#if WITH_VERSE
+					if (Cast<UVerseClassBase>(BaseClass) != nullptr)
+					{
+						// Verse devices cannot have C++ or blueprint classes created that derive from them
+						return;
+					}
+#endif
+
 					auto CreateCreateDerivedCppClass = [BaseClass](const FToolMenuContext& InContext)
 					{
 						// Work out where the header file for the current class is, as we'll use that path as the default for the new class
