[Backout] - CL24666850
To address Main CIS Non Unity Error.

Original CL Desc
-----------------------------------------------------------------
Disable the "Create Blueprint class" based on non bp script

#rnx

diff --git a/Engine/Plugins/Editor/EngineAssetDefinitions/Source/EngineAssetDefinitions.Build.cs b/Engine/Plugins/Editor/EngineAssetDefinitions/Source/EngineAssetDefinitions.Build.cs
--- a/Engine/Plugins/Editor/EngineAssetDefinitions/Source/EngineAssetDefinitions.Build.cs
+++ b/Engine/Plugins/Editor/EngineAssetDefinitions/Source/EngineAssetDefinitions.Build.cs
@@ -40,15 +40,7 @@ public class EngineAssetDefinitions : ModuleRules
 				"ToolWidgets"
 			}
 		);
-
-		if (Target.bUseVerse)
-		{
-			PrivateDependencyModuleNames.AddRange(new string[]
-			{
-				"VerseNative",
-			});
-		}
-
+		
 		PrivateIncludePathModuleNames.AddRange(
 			new string[] {
 				"VirtualTexturingEditor",
diff --git a/Engine/Plugins/Editor/EngineAssetDefinitions/Source/Private/Script/AssetDefinition_Class.cpp b/Engine/Plugins/Editor/EngineAssetDefinitions/Source/Private/Script/AssetDefinition_Class.cpp
--- a/Engine/Plugins/Editor/EngineAssetDefinitions/Source/Private/Script/AssetDefinition_Class.cpp
+++ b/Engine/Plugins/Editor/EngineAssetDefinitions/Source/Private/Script/AssetDefinition_Class.cpp
@@ -17,12 +17,9 @@
 #include "ToolMenuSection.h"
 #include "UnrealEdGlobals.h"
 
-#if WITH_VERSE
-#include "VerseClassBase.h"
-#endif
-
 #define LOCTEXT_NAMESPACE "AssetTypeActions"
 
+
 UThumbnailInfo* UAssetDefinition_Class::LoadThumbnailInfo(const FAssetData& InAssetData) const
 {
 	// todo: jdale - CLASS - We need to generate and store proper thumbnail info for classes so that we can store their custom render transforms
@@ -95,14 +92,6 @@ namespace MenuExtension_Class
 					const bool bIsValidBaseCppClass = BaseClass && GameProjectGenerationModule.IsValidBaseClassForCreation(BaseClass, ProjectModules);
 					const bool bIsValidBaseBlueprintClass = BaseClass && FKismetEditorUtilities::CanCreateBlueprintOfClass(BaseClass);
 
-#if WITH_VERSE
-					if (Cast<UVerseClassBase>(BaseClass) != nullptr)
-					{
-						// Verse devices cannot have C++ or blueprint classes created that derive from them
-						return;
-					}
-#endif
-
 					auto CreateCreateDerivedCppClass = [BaseClass](const FToolMenuContext& InContext)
 					{
 						// Work out where the header file for the current class is, as we'll use that path as the default for the new class
