Disable the "Create Blueprint class" for assets based on non bp scripts
This is done by checking the asset tags on the UObject

diff --git a/Engine/Plugins/Editor/EngineAssetDefinitions/Source/Private/Script/AssetDefinition_Class.cpp b/Engine/Plugins/Editor/EngineAssetDefinitions/Source/Private/Script/AssetDefinition_Class.cpp
--- a/Engine/Plugins/Editor/EngineAssetDefinitions/Source/Private/Script/AssetDefinition_Class.cpp
+++ b/Engine/Plugins/Editor/EngineAssetDefinitions/Source/Private/Script/AssetDefinition_Class.cpp
@@ -86,12 +86,28 @@ namespace MenuExtension_Class
 				{
 					UClass *const BaseClass = (CBContext->SelectedAssets.Num() == 1) ? Cast<UClass>(CBContext->SelectedAssets[0].GetAsset()) : nullptr;
 
+#if UE_USE_VERSE_PATHS
+					if (BaseClass)
+					{
+						TArray<UObject::FAssetRegistryTag> InTags;
+						InTags.Reserve(32);
+						BaseClass->GetAssetRegistryTags(InTags);
+						for (const UObject::FAssetRegistryTag& CurrentTag : InTags)
+						{
+							if (CurrentTag.Name == UObject::AssetVersePathTagName())
+							{
+								// Verse devices cannot have C++ or blueprint classes created that derive from them
+								return;
+							}
+						}
+					}
+#endif
+
 					// Only allow the New class option if we have a base class that we can actually derive from in one of our project modules
 					FGameProjectGenerationModule& GameProjectGenerationModule = FGameProjectGenerationModule::Get();
 					TArray<FModuleContextInfo> ProjectModules = GameProjectGenerationModule.GetCurrentProjectModules();
 					const bool bIsValidBaseCppClass = BaseClass && GameProjectGenerationModule.IsValidBaseClassForCreation(BaseClass, ProjectModules);
 					const bool bIsValidBaseBlueprintClass = BaseClass && FKismetEditorUtilities::CanCreateBlueprintOfClass(BaseClass);
-
 					auto CreateCreateDerivedCppClass = [BaseClass](const FToolMenuContext& InContext)
 					{
 						// Work out where the header file for the current class is, as we'll use that path as the default for the new class
