UMG + EUB AssetDefinition conversion

This change converts EditorUtilityWidgetBlueprint, EditorUtilityBlueprint and UMG's WidgetBlueprint and WidgetBlueprintGeneratedClass assets to use the newer AssetDefintion framework.
AssetTypeActions were removed with this CL.

diff --git a/Engine/Plugins/Editor/EngineAssetDefinitions/Source/EngineAssetDefinitions.Build.cs b/Engine/Plugins/Editor/EngineAssetDefinitions/Source/EngineAssetDefinitions.Build.cs
--- a/Engine/Plugins/Editor/EngineAssetDefinitions/Source/EngineAssetDefinitions.Build.cs
+++ b/Engine/Plugins/Editor/EngineAssetDefinitions/Source/EngineAssetDefinitions.Build.cs
@@ -37,6 +37,8 @@ public class EngineAssetDefinitions : ModuleRules
 				"Kismet",
 				"GameProjectGeneration",
 				"AnimationBlueprintEditor",
+				"Blutility",
+				"UMGEditor",
 				"ToolWidgets"
 			}
 		);
diff --git a/Engine/Plugins/Editor/EngineAssetDefinitions/Source/Private/AssetDefinition_EditorUtilityBlueprint.cpp b/Engine/Plugins/Editor/EngineAssetDefinitions/Source/Private/AssetDefinition_EditorUtilityBlueprint.cpp
new file mode 100644
--- /dev/null
+++ b/Engine/Plugins/Editor/EngineAssetDefinitions/Source/Private/AssetDefinition_EditorUtilityBlueprint.cpp
@@ -0,0 +1,105 @@
+﻿// Copyright Epic Games, Inc. All Rights Reserved.
+
+#include "AssetDefinition_EditorUtilityBlueprint.h"
+
+#include "ContentBrowserMenuContexts.h"
+
+#include "Containers/ArrayView.h"
+#include "Editor.h"
+#include "EditorUtilityBlueprint.h"
+#include "EditorUtilitySubsystem.h"
+#include "IAssetTools.h"
+#include "IBlutilityModule.h"
+#include "ObjectEditorUtils.h"
+#include "ToolMenus.h"
+#include "Algo/RemoveIf.h"
+
+#define LOCTEXT_NAMESPACE "AssetTypeActions"
+
+FText UAssetDefinition_EditorUtilityBlueprint::GetAssetDisplayName() const
+{
+	return LOCTEXT("AssetTypeActions_EditorUtilityBlueprintUpdate", "Editor Utility Blueprint");
+}
+
+FLinearColor UAssetDefinition_EditorUtilityBlueprint::GetAssetColor() const
+{
+	return FColor(0, 169, 255);
+}
+
+TSoftClassPtr<> UAssetDefinition_EditorUtilityBlueprint::GetAssetClass() const
+{
+	return UEditorUtilityBlueprint::StaticClass();
+}
+
+TConstArrayView<FAssetCategoryPath> UAssetDefinition_EditorUtilityBlueprint::GetAssetCategories() const
+{
+	IBlutilityModule* BlutilityModule = FModuleManager::GetModulePtr<IBlutilityModule>("Blutility");
+	if (BlutilityModule)
+	{
+		return BlutilityModule->GetAssetCategories();
+	}
+
+	static const FAssetCategoryPath FallbackCategory[] = { EAssetCategoryPaths::Misc };
+	return FallbackCategory;
+}
+
+namespace MenuExtension_EditorUtilityBlueprint
+{
+	using FWeakBlueprintPointerArray = TArray< TWeakObjectPtr<class UWidgetBlueprint> >;
+	
+	static void ExecuteEditorUtilityRun(const FToolMenuContext& InContext, TArray<UObject*> Objects)
+	{
+		if (const UContentBrowserAssetContextMenuContext* Context = UContentBrowserAssetContextMenuContext::FindContextWithAssets(InContext))
+		{
+			UEditorUtilitySubsystem* EditorUtilitySubsystem = GEditor->GetEditorSubsystem<UEditorUtilitySubsystem>();
+			for (UObject* Object : Objects)
+			{
+				EditorUtilitySubsystem->TryRun(Object);
+			}
+		}
+	}
+	
+	static FDelayedAutoRegisterHelper DelayedAutoRegister(EDelayedRegisterRunPhase::EndOfEngineInit,[]
+	{
+		UToolMenus::RegisterStartupCallback(FSimpleMulticastDelegate::FDelegate::CreateLambda([]()
+		{
+			FToolMenuOwnerScoped OwnerScoped(UE_MODULE_NAME);
+			UToolMenu* Menu = UE::ContentBrowser::ExtendToolMenu_AssetContextMenu(UEditorUtilityBlueprint::StaticClass());
+
+			FToolMenuSection& Section = Menu->FindOrAddSection("GetAssetActions");
+			Section.AddDynamicEntry(NAME_None, FNewToolMenuSectionDelegate::CreateLambda([](FToolMenuSection& InSection)
+			{
+				if (const UContentBrowserAssetContextMenuContext* Context = UContentBrowserAssetContextMenuContext::FindContextWithAssets(InSection))
+				{
+					UEditorUtilitySubsystem* EUSubsystem = GEditor->GetEditorSubsystem<UEditorUtilitySubsystem>();
+
+					// It would be ideal to not load selected object here, instead to defer the load to
+					// ExecuteEditorUtilityRun(). Unfortunately, the GeneratedClass of a UBlueprint
+					// is not accessible from AssetData without loading the asset. This is required for the CanRun()
+					// check.  Results in an asset load on opening the context menu with these selected.
+					TSet<FName> LoadTags;
+					TArray<UObject*> SelectedAssets = Context->LoadSelectedObjects(LoadTags);
+
+					const int32 EndOfRange = Algo::StableRemoveIf(SelectedAssets, [&EUSubsystem](UObject* Object)
+					{
+						return !EUSubsystem->CanRun(Object);
+					});
+					SelectedAssets.SetNum(EndOfRange);
+
+					if (!SelectedAssets.IsEmpty())
+					{
+						const TAttribute<FText> Label = LOCTEXT("EditorUtility_Run", "Run Editor Utility Blueprint");
+						const TAttribute<FText> ToolTip = LOCTEXT("EditorUtility_RunTooltip", "Runs this Editor Utility Blueprint.");
+						const FSlateIcon Icon = FSlateIcon();
+						
+						FToolUIAction UIAction;
+						UIAction.ExecuteAction = FToolMenuExecuteAction::CreateStatic(&ExecuteEditorUtilityRun, MoveTemp(SelectedAssets));
+						InSection.AddMenuEntry("EditorUtility_Run", Label, ToolTip, Icon, UIAction);
+					}
+				}
+			}));
+		}));
+	});
+}
+
+#undef LOCTEXT_NAMESPACE
\ No newline at end of file
diff --git a/Engine/Plugins/Editor/EngineAssetDefinitions/Source/Private/AssetDefinition_EditorUtilityBlueprint.h b/Engine/Plugins/Editor/EngineAssetDefinitions/Source/Private/AssetDefinition_EditorUtilityBlueprint.h
new file mode 100644
--- /dev/null
+++ b/Engine/Plugins/Editor/EngineAssetDefinitions/Source/Private/AssetDefinition_EditorUtilityBlueprint.h
@@ -0,0 +1,21 @@
+﻿// Copyright Epic Games, Inc. All Rights Reserved.
+
+#pragma once
+
+#include "Script/AssetDefinition_Blueprint.h"
+
+#include "AssetDefinition_EditorUtilityBlueprint.generated.h"
+
+UCLASS()
+class UAssetDefinition_EditorUtilityBlueprint : public UAssetDefinition_Blueprint
+{
+	GENERATED_BODY()
+public:
+	//~ UAssetDefinition_Blueprint Begin
+	virtual FText GetAssetDisplayName() const override;
+	virtual FLinearColor GetAssetColor() const override;
+	virtual TSoftClassPtr<UObject> GetAssetClass() const override;
+	virtual TConstArrayView<FAssetCategoryPath> GetAssetCategories() const override;
+	//~ UAssetDefinition_Blueprint End
+	
+};
diff --git a/Engine/Plugins/Editor/EngineAssetDefinitions/Source/Private/AssetDefinition_EditorUtilityWidgetBlueprint.cpp b/Engine/Plugins/Editor/EngineAssetDefinitions/Source/Private/AssetDefinition_EditorUtilityWidgetBlueprint.cpp
new file mode 100644
--- /dev/null
+++ b/Engine/Plugins/Editor/EngineAssetDefinitions/Source/Private/AssetDefinition_EditorUtilityWidgetBlueprint.cpp
@@ -0,0 +1,139 @@
+﻿// Copyright Epic Games, Inc. All Rights Reserved.
+
+#include "AssetDefinition_EditorUtilityWidgetBlueprint.h"
+
+#include "Algo/RemoveIf.h"
+#include "ContentBrowserMenuContexts.h"
+#include "EditorUtilitySubsystem.h"
+#include "EditorUtilityWidgetBlueprint.h"
+#include "IBlutilityModule.h"
+#include "Misc/MessageDialog.h"
+#include "SBlueprintDiff.h"
+#include "WidgetBlueprintEditor.h"
+
+#define LOCTEXT_NAMESPACE "AssetTypeActions"
+
+UAssetDefinition_EditorUtilityWidgetBlueprint::UAssetDefinition_EditorUtilityWidgetBlueprint() = default;
+
+UAssetDefinition_EditorUtilityWidgetBlueprint::~UAssetDefinition_EditorUtilityWidgetBlueprint() = default;
+
+FText UAssetDefinition_EditorUtilityWidgetBlueprint::GetAssetDisplayName() const
+{
+	return LOCTEXT("AssetTypeActions_EditorUtilityWidget", "Editor Utility Widget");
+}
+
+FLinearColor UAssetDefinition_EditorUtilityWidgetBlueprint::GetAssetColor() const
+{
+	return FColor(0, 169, 255);
+}
+
+TSoftClassPtr<> UAssetDefinition_EditorUtilityWidgetBlueprint::GetAssetClass() const
+{
+	return UEditorUtilityWidgetBlueprint::StaticClass();
+}
+
+TConstArrayView<FAssetCategoryPath> UAssetDefinition_EditorUtilityWidgetBlueprint::GetAssetCategories() const
+{
+	IBlutilityModule* BlutilityModule = FModuleManager::GetModulePtr<IBlutilityModule>("Blutility");
+	if (BlutilityModule)
+	{
+		return BlutilityModule->GetAssetCategories();
+	}
+
+	static const TArray<FAssetCategoryPath, TFixedAllocator<1>> FallbackCategory = { EAssetCategoryPaths::Misc };
+	return FallbackCategory;
+}
+
+EAssetCommandResult UAssetDefinition_EditorUtilityWidgetBlueprint::OpenAssets(const FAssetOpenArgs& OpenArgs) const
+{
+	EToolkitMode::Type Mode = OpenArgs.GetToolkitMode();
+
+	for (UBlueprint* Blueprint : OpenArgs.LoadObjects<UBlueprint>())
+	{
+		if (Blueprint->SkeletonGeneratedClass && Blueprint->GeneratedClass)
+		{
+			TSharedRef<FWidgetBlueprintEditor> NewBlueprintEditor(new FWidgetBlueprintEditor());
+
+			TArray<UBlueprint*> Blueprints;
+			Blueprints.Add(Blueprint);
+			NewBlueprintEditor->InitWidgetBlueprintEditor(EToolkitMode::Standalone, nullptr, Blueprints, true);
+		}
+		else
+		{
+			FMessageDialog::Open(EAppMsgType::Ok, LOCTEXT("FailedToLoadEditorUtilityWidgetBlueprint", "Editor Utility Widget could not be loaded because it derives from an invalid class.\nCheck to make sure the parent class for this blueprint hasn't been removed!"));	
+		}
+	}
+	return EAssetCommandResult::Handled;
+}
+
+EAssetCommandResult UAssetDefinition_EditorUtilityWidgetBlueprint::PerformAssetDiff(const FAssetDiffArgs& DiffArgs) const
+{
+	const UBlueprint* OldBlueprint = Cast<UBlueprint>(DiffArgs.OldAsset);
+	const UBlueprint* NewBlueprint = Cast<UBlueprint>(DiffArgs.NewAsset);
+	SBlueprintDiff::CreateDiffWindow(OldBlueprint, NewBlueprint, DiffArgs.OldRevision, DiffArgs.NewRevision, GetAssetClass().Get());
+	return EAssetCommandResult::Handled;
+}
+
+namespace MenuExtension_EditorUtilityWidgetBlueprint
+{
+	void ExecuteEditorUtilityEdit(const FToolMenuContext& InContext, TArray<FAssetData> AssetData)
+	{
+		if (const UContentBrowserAssetContextMenuContext* Context = UContentBrowserAssetContextMenuContext::FindContextWithAssets(InContext))
+		{
+			UEditorUtilitySubsystem* EditorUtilitySubsystem = GEditor->GetEditorSubsystem<UEditorUtilitySubsystem>();
+			
+			TSet<FName> LoadTags;
+			TArray<UWidgetBlueprint*> WidgetBlueprints = Context->LoadSelectedObjects<UWidgetBlueprint>(LoadTags);
+			
+			for (UWidgetBlueprint* WidgetBlueprint : WidgetBlueprints)
+			{
+				ensureMsgf(WidgetBlueprint->GeneratedClass, TEXT("Null generated class for WidgetBlueprint [%s]"), *WidgetBlueprint->GetName());
+				if (WidgetBlueprint->GeneratedClass && WidgetBlueprint->GeneratedClass->IsChildOf(UEditorUtilityWidget::StaticClass()))
+				{
+					if (UEditorUtilityWidgetBlueprint* EditorWidget = Cast<UEditorUtilityWidgetBlueprint>(WidgetBlueprint))
+					{
+						EditorUtilitySubsystem->SpawnAndRegisterTab(EditorWidget);
+					}
+				}
+			}
+		}
+	}
+	
+	FDelayedAutoRegisterHelper DelayedAutoRegister(EDelayedRegisterRunPhase::EndOfEngineInit,[]
+	{
+		UToolMenus::RegisterStartupCallback(FSimpleMulticastDelegate::FDelegate::CreateLambda([]()
+		{
+			FToolMenuOwnerScoped OwnerScoped(UE_MODULE_NAME);
+			UToolMenu* Menu = UE::ContentBrowser::ExtendToolMenu_AssetContextMenu(UEditorUtilityWidgetBlueprint::StaticClass());
+
+			FToolMenuSection& Section = Menu->FindOrAddSection("GetAssetActions");
+			Section.AddDynamicEntry(NAME_None, FNewToolMenuSectionDelegate::CreateLambda([](FToolMenuSection& InSection)
+			{
+				if (const UContentBrowserAssetContextMenuContext* Context = UContentBrowserAssetContextMenuContext::FindContextWithAssets(InSection))
+				{
+					TArray<FAssetData> SelectedAssets(Context->SelectedAssets);
+				
+					const int32 EndOfRange = Algo::StableRemoveIf(SelectedAssets, [](const FAssetData& AssetData)
+					{
+						return !AssetData.IsInstanceOf(UWidgetBlueprint::StaticClass());
+					});
+
+					SelectedAssets.SetNum(EndOfRange);
+
+					if (!SelectedAssets.IsEmpty())
+					{
+						const TAttribute<FText> Label = LOCTEXT("EditorUtilityWidget_Edit", "Run Editor Utility Widget");
+						const TAttribute<FText> ToolTip = LOCTEXT("EditorUtilityWidget_EditTooltip", "Opens the tab built by this Editor Utility Widget Blueprint.");
+						const FSlateIcon Icon = FSlateIcon();
+						
+						FToolUIAction UIAction;
+						UIAction.ExecuteAction = FToolMenuExecuteAction::CreateStatic(&ExecuteEditorUtilityEdit, MoveTemp(SelectedAssets));
+						InSection.AddMenuEntry("EditorUtility_Run", Label, ToolTip, Icon, UIAction);
+					}
+				}
+			}));
+		}));
+	});
+}
+
+#undef LOCTEXT_NAMESPACE
diff --git a/Engine/Plugins/Editor/EngineAssetDefinitions/Source/Private/AssetDefinition_EditorUtilityWidgetBlueprint.h b/Engine/Plugins/Editor/EngineAssetDefinitions/Source/Private/AssetDefinition_EditorUtilityWidgetBlueprint.h
new file mode 100644
--- /dev/null
+++ b/Engine/Plugins/Editor/EngineAssetDefinitions/Source/Private/AssetDefinition_EditorUtilityWidgetBlueprint.h
@@ -0,0 +1,29 @@
+﻿// Copyright Epic Games, Inc. All Rights Reserved.
+
+#pragma once
+
+#include "AssetDefinitionDefault.h"
+
+#include "AssetDefinition_EditorUtilityWidgetBlueprint.generated.h"
+
+UCLASS()
+class UAssetDefinition_EditorUtilityWidgetBlueprint : public UAssetDefinitionDefault
+{
+	GENERATED_BODY()
+public:
+	UAssetDefinition_EditorUtilityWidgetBlueprint();
+	virtual ~UAssetDefinition_EditorUtilityWidgetBlueprint() override;
+
+	//~ UAssetDefinition Begin
+	virtual FText GetAssetDisplayName() const override;
+	virtual FLinearColor GetAssetColor() const override;
+	virtual TSoftClassPtr<UObject> GetAssetClass() const override;
+	virtual TConstArrayView<FAssetCategoryPath> GetAssetCategories() const override;
+
+	virtual EAssetCommandResult OpenAssets(const FAssetOpenArgs& OpenArgs) const override;
+	virtual EAssetCommandResult PerformAssetDiff(const FAssetDiffArgs& DiffArgs) const override;
+	//~ UAssetDefinition End
+	
+};
+
+
diff --git a/Engine/Source/Editor/Blutility/Blutility.Build.cs b/Engine/Source/Editor/Blutility/Blutility.Build.cs
--- a/Engine/Source/Editor/Blutility/Blutility.Build.cs
+++ b/Engine/Source/Editor/Blutility/Blutility.Build.cs
@@ -24,6 +24,7 @@ public class Blutility : ModuleRules
 				"EditorFramework",
 				"UnrealEd",
 				"Kismet",
+				"AssetDefinition",
 				"AssetRegistry",
 				"AssetTools",
 				"WorkspaceMenuStructure",
diff --git a/Engine/Source/Editor/Blutility/Private/AssetTypeActions_EditorUtilityBlueprint.cpp b/Engine/Source/Editor/Blutility/Private/AssetTypeActions_EditorUtilityBlueprint.cpp
deleted file mode 100644
diff --git a/Engine/Source/Editor/Blutility/Private/AssetTypeActions_EditorUtilityBlueprint.h b/Engine/Source/Editor/Blutility/Private/AssetTypeActions_EditorUtilityBlueprint.h
deleted file mode 100644
diff --git a/Engine/Source/Editor/Blutility/Private/AssetTypeActions_EditorUtilityWidgetBlueprint.cpp b/Engine/Source/Editor/Blutility/Private/AssetTypeActions_EditorUtilityWidgetBlueprint.cpp
deleted file mode 100644
diff --git a/Engine/Source/Editor/Blutility/Private/AssetTypeActions_EditorUtilityWidgetBlueprint.h b/Engine/Source/Editor/Blutility/Private/AssetTypeActions_EditorUtilityWidgetBlueprint.h
deleted file mode 100644
diff --git a/Engine/Source/Editor/Blutility/Private/BlutilityModule.cpp b/Engine/Source/Editor/Blutility/Private/BlutilityModule.cpp
--- a/Engine/Source/Editor/Blutility/Private/BlutilityModule.cpp
+++ b/Engine/Source/Editor/Blutility/Private/BlutilityModule.cpp
@@ -4,8 +4,6 @@
 #include "AssetRegistry/AssetRegistryModule.h"
 #include "AssetRegistry/IAssetRegistry.h"
 #include "AssetToolsModule.h"
-#include "AssetTypeActions_EditorUtilityBlueprint.h"
-#include "AssetTypeActions_EditorUtilityWidgetBlueprint.h"
 #include "AssetTypeCategories.h"
 #include "BlutilityContentBrowserExtensions.h"
 #include "BlutilityLevelEditorExtensions.h"
@@ -77,21 +75,12 @@ DEFINE_LOG_CATEGORY(LogEditorUtilityBlueprint);
 // Blutility module implementation (private)
 class FBlutilityModule : public IBlutilityModule, public FGCObject
 {
-public:
-	/** Asset type actions for editor utility assets.  Cached here so that we can unregister it during shutdown. */
-	TSharedPtr<FAssetTypeActions_EditorUtilityBlueprint> EditorBlueprintAssetTypeActions;
-	TSharedPtr<FAssetTypeActions_EditorUtilityWidgetBlueprint> EditorWidgetBlueprintAssetTypeActions;
-
 public:
 	virtual void StartupModule() override
 	{
 		// Register the asset type
 		IAssetTools& AssetTools = FModuleManager::LoadModuleChecked<FAssetToolsModule>("AssetTools").Get();
 		EditorUtilityAssetCategory = AssetTools.RegisterAdvancedAssetCategory(FName(TEXT("EditorUtilities")), LOCTEXT("EditorUtilitiesAssetCategory", "Editor Utilities"));
-		EditorBlueprintAssetTypeActions = MakeShareable(new FAssetTypeActions_EditorUtilityBlueprint);
-		AssetTools.RegisterAssetTypeActions(EditorBlueprintAssetTypeActions.ToSharedRef());
-		EditorWidgetBlueprintAssetTypeActions = MakeShareable(new FAssetTypeActions_EditorUtilityWidgetBlueprint);
-		AssetTools.RegisterAssetTypeActions(EditorWidgetBlueprintAssetTypeActions.ToSharedRef());
 
 		FKismetCompilerContext::RegisterCompilerForBP(UEditorUtilityWidgetBlueprint::StaticClass(), &UWidgetBlueprint::GetCompilerForWidgetBP);
 
@@ -208,17 +197,6 @@ public:
 		FBlutilityContentBrowserExtensions::RemoveHooks();
 		FBlutilityUMGEditorExtensions::RemoveHooks();
 
-		// Only unregister if the asset tools module is loaded.  We don't want to forcibly load it during shutdown phase.
-		check( EditorBlueprintAssetTypeActions.IsValid() );
-		if (FModuleManager::Get().IsModuleLoaded("AssetTools"))
-		{
-			IAssetTools& AssetTools = FModuleManager::LoadModuleChecked<FAssetToolsModule>("AssetTools").Get();
-			AssetTools.UnregisterAssetTypeActions(EditorBlueprintAssetTypeActions.ToSharedRef());
-			AssetTools.UnregisterAssetTypeActions(EditorWidgetBlueprintAssetTypeActions.ToSharedRef());
-		}
-		EditorBlueprintAssetTypeActions.Reset();
-		EditorWidgetBlueprintAssetTypeActions.Reset();
-
 		FEditorSupportDelegates::PrepareToCleanseEditorObject.RemoveAll(this);
 	}
 
@@ -245,6 +223,14 @@ public:
 		return EditorUtilityAssetCategory;
 	}
 
+	virtual TConstArrayView<FAssetCategoryPath> GetAssetCategories() const override
+	{
+		static const TArray<FAssetCategoryPath, TFixedAllocator<1>> Categories = {
+			FAssetCategoryPath(LOCTEXT("EditorUtilities", "EditorUtilities"))
+		};
+		return Categories;
+	}
+
 	virtual void AddLoadedScriptUI(class UEditorUtilityWidgetBlueprint* InBlueprint) override
 	{
 		UEditorUtilitySubsystem* EditorUtilitySubsystem = GEditor->GetEditorSubsystem<UEditorUtilitySubsystem>();
@@ -252,7 +238,6 @@ public:
 		EditorUtilitySubsystem->SaveConfig();
 	}
 
-
 	virtual void RemoveLoadedScriptUI(class UEditorUtilityWidgetBlueprint* InBlueprint) override
 	{
 		UEditorUtilitySubsystem* EditorUtilitySubsystem = GEditor->GetEditorSubsystem<UEditorUtilitySubsystem>();
diff --git a/Engine/Source/Editor/Blutility/Public/IBlutilityModule.h b/Engine/Source/Editor/Blutility/Public/IBlutilityModule.h
--- a/Engine/Source/Editor/Blutility/Public/IBlutilityModule.h
+++ b/Engine/Source/Editor/Blutility/Public/IBlutilityModule.h
@@ -6,6 +6,7 @@
 #include "Modules/ModuleInterface.h"
 #include "AssetTypeCategories.h"
 
+struct FAssetCategoryPath;
 class UBlueprint;
 
 /**
@@ -21,8 +22,11 @@ public:
 	/** Global Find Results workspace menu item */
 	virtual TSharedPtr<class FWorkspaceItem> GetMenuGroup() const = 0;
 
+	UE_DEPRECATED(5.3, "Use GetAssetCategories instead")
 	virtual EAssetTypeCategories::Type GetAssetCategory() const = 0;
 
+	virtual TConstArrayView<FAssetCategoryPath> GetAssetCategories() const = 0;
+
 	virtual void AddLoadedScriptUI(class UEditorUtilityWidgetBlueprint* InBlueprint) = 0;
 
 	virtual void RemoveLoadedScriptUI(class UEditorUtilityWidgetBlueprint* InBlueprint) = 0;
diff --git a/Engine/Source/Editor/UMGEditor/Private/AssetDefinition_WidgetBlueprint.cpp b/Engine/Source/Editor/UMGEditor/Private/AssetDefinition_WidgetBlueprint.cpp
new file mode 100644
--- /dev/null
+++ b/Engine/Source/Editor/UMGEditor/Private/AssetDefinition_WidgetBlueprint.cpp
@@ -0,0 +1,86 @@
+// Copyright Epic Games, Inc. All Rights Reserved.
+
+#include "AssetDefinition_WidgetBlueprint.h"
+
+#include "WidgetBlueprintEditor.h"
+#include "Misc/MessageDialog.h"
+#include "SBlueprintDiff.h"
+
+#define LOCTEXT_NAMESPACE "AssetTypeActions"
+
+UAssetDefinition_WidgetBlueprint::UAssetDefinition_WidgetBlueprint() = default;
+
+UAssetDefinition_WidgetBlueprint::~UAssetDefinition_WidgetBlueprint() = default;
+
+FText UAssetDefinition_WidgetBlueprint::GetAssetDisplayName() const
+{
+	return LOCTEXT("AssetTypeActions_WidgetBlueprint", "Widget Blueprint");
+}
+
+FLinearColor UAssetDefinition_WidgetBlueprint::GetAssetColor() const
+{
+	return FLinearColor(FColor(44, 89, 180));
+}
+
+TSoftClassPtr<> UAssetDefinition_WidgetBlueprint::GetAssetClass() const
+{
+	return UWidgetBlueprint::StaticClass();
+}
+
+TConstArrayView<FAssetCategoryPath> UAssetDefinition_WidgetBlueprint::GetAssetCategories() const
+{
+	static const TArray<FAssetCategoryPath, TFixedAllocator<1>> Categories = { EAssetCategoryPaths::UI };
+	return Categories;
+}
+
+EAssetCommandResult UAssetDefinition_WidgetBlueprint::OpenAssets(const FAssetOpenArgs& OpenArgs) const
+{
+	EToolkitMode::Type Mode = OpenArgs.GetToolkitMode();
+
+	EAssetCommandResult Result = EAssetCommandResult::Unhandled;
+
+	for (UBlueprint* Blueprint : OpenArgs.LoadObjects<UBlueprint>())
+	{
+		if (Blueprint && Blueprint->SkeletonGeneratedClass && Blueprint->GeneratedClass)
+		{
+			TSharedRef<FWidgetBlueprintEditor> NewBlueprintEditor(new FWidgetBlueprintEditor);
+
+			const bool bShouldOpenInDefaultsMode = false;
+			TArray<UBlueprint*> Blueprints;
+			Blueprints.Add(Blueprint);
+
+			NewBlueprintEditor->InitWidgetBlueprintEditor(Mode, OpenArgs.ToolkitHost, Blueprints, bShouldOpenInDefaultsMode);
+		}
+		else
+		{
+			FMessageDialog::Open( EAppMsgType::Ok, LOCTEXT("FailedToLoadWidgetBlueprint", "Widget Blueprint could not be loaded because it derives from an invalid class.\nCheck to make sure the parent class for this blueprint hasn't been removed!"));
+		}
+
+		Result = EAssetCommandResult::Handled;
+	}
+
+	return Result;
+}
+
+EAssetCommandResult UAssetDefinition_WidgetBlueprint::PerformAssetDiff(const FAssetDiffArgs& DiffArgs) const
+{
+	const UBlueprint* OldBlueprint = Cast<UBlueprint>(DiffArgs.OldAsset);
+	const UBlueprint* NewBlueprint = Cast<UBlueprint>(DiffArgs.NewAsset);
+	UClass* AssetClass = GetAssetClass().Get();
+	SBlueprintDiff::CreateDiffWindow(OldBlueprint, NewBlueprint, DiffArgs.OldRevision, DiffArgs.NewRevision, AssetClass);
+	return EAssetCommandResult::Handled;
+}
+
+FText UAssetDefinition_WidgetBlueprint::GetAssetDescription(const FAssetData& AssetData) const
+{
+	FString Description = AssetData.GetTagValueRef<FString>( GET_MEMBER_NAME_CHECKED( UBlueprint, BlueprintDescription ) );
+	if ( !Description.IsEmpty() )
+	{
+		Description.ReplaceInline( TEXT( "\\n" ), TEXT( "\n" ) );
+		return FText::FromString( MoveTemp(Description) );
+	}
+
+	return FText::GetEmpty();
+}
+
+#undef LOCTEXT_NAMESPACE
diff --git a/Engine/Source/Editor/UMGEditor/Private/AssetDefinition_WidgetBlueprint.h b/Engine/Source/Editor/UMGEditor/Private/AssetDefinition_WidgetBlueprint.h
new file mode 100644
--- /dev/null
+++ b/Engine/Source/Editor/UMGEditor/Private/AssetDefinition_WidgetBlueprint.h
@@ -0,0 +1,28 @@
+﻿// Copyright Epic Games, Inc. All Rights Reserved.
+
+#pragma once
+
+#include "WidgetBlueprint.h"
+#include "AssetDefinitionDefault.h"
+
+#include "AssetDefinition_WidgetBlueprint.generated.h"
+
+UCLASS()
+class UAssetDefinition_WidgetBlueprint : public UAssetDefinitionDefault
+{
+	GENERATED_BODY()
+
+public:
+	UAssetDefinition_WidgetBlueprint();
+	virtual ~UAssetDefinition_WidgetBlueprint() override;
+	
+	// UAssetDefinition Begin
+	virtual FText GetAssetDisplayName() const override;
+	virtual FLinearColor GetAssetColor() const override;
+	virtual TSoftClassPtr<UObject> GetAssetClass() const override;
+	virtual TConstArrayView<FAssetCategoryPath> GetAssetCategories() const override;
+	virtual EAssetCommandResult OpenAssets(const FAssetOpenArgs& OpenArgs) const override;
+	virtual EAssetCommandResult PerformAssetDiff(const FAssetDiffArgs& DiffArgs) const override;
+	virtual FText GetAssetDescription(const FAssetData& AssetData) const override;
+	// UAssetDefinition End
+};
diff --git a/Engine/Source/Editor/UMGEditor/Private/AssetDefinition_WidgetBlueprintGeneratedClass.cpp b/Engine/Source/Editor/UMGEditor/Private/AssetDefinition_WidgetBlueprintGeneratedClass.cpp
new file mode 100644
--- /dev/null
+++ b/Engine/Source/Editor/UMGEditor/Private/AssetDefinition_WidgetBlueprintGeneratedClass.cpp
@@ -0,0 +1,28 @@
+﻿// Copyright Epic Games, Inc. All Rights Reserved.
+
+#include "AssetDefinition_WidgetBlueprintGeneratedClass.h"
+
+UAssetDefinition_WidgetBlueprintGeneratedClass::UAssetDefinition_WidgetBlueprintGeneratedClass() = default;
+
+UAssetDefinition_WidgetBlueprintGeneratedClass::~UAssetDefinition_WidgetBlueprintGeneratedClass() = default;
+
+FText UAssetDefinition_WidgetBlueprintGeneratedClass::GetAssetDisplayName() const
+{
+	return NSLOCTEXT("AssetTypeActions", "AssetTypeActions_WidgetBlueprintGeneratedClass", "Compiled Widget Blueprint");
+}
+
+FLinearColor UAssetDefinition_WidgetBlueprintGeneratedClass::GetAssetColor() const
+{
+	return FColor(121,149,207);
+}
+
+TSoftClassPtr<> UAssetDefinition_WidgetBlueprintGeneratedClass::GetAssetClass() const
+{
+	return UWidgetBlueprintGeneratedClass::StaticClass();
+}
+
+TConstArrayView<FAssetCategoryPath> UAssetDefinition_WidgetBlueprintGeneratedClass::GetAssetCategories() const
+{
+	static const auto Categories = { EAssetCategoryPaths::UI };
+	return Categories;
+}
diff --git a/Engine/Source/Editor/UMGEditor/Private/AssetDefinition_WidgetBlueprintGeneratedClass.h b/Engine/Source/Editor/UMGEditor/Private/AssetDefinition_WidgetBlueprintGeneratedClass.h
new file mode 100644
--- /dev/null
+++ b/Engine/Source/Editor/UMGEditor/Private/AssetDefinition_WidgetBlueprintGeneratedClass.h
@@ -0,0 +1,25 @@
+﻿// Copyright Epic Games, Inc. All Rights Reserved.
+
+#pragma once
+
+#include "WidgetBlueprint.h"
+#include "AssetDefinitionDefault.h"
+
+#include "AssetDefinition_WidgetBlueprintGeneratedClass.generated.h"
+
+UCLASS()
+class UAssetDefinition_WidgetBlueprintGeneratedClass : public UAssetDefinitionDefault
+{
+	GENERATED_BODY()
+
+public:
+	UAssetDefinition_WidgetBlueprintGeneratedClass();
+	virtual ~UAssetDefinition_WidgetBlueprintGeneratedClass() override;
+	
+	// UAssetDefinition Begin
+	virtual FText GetAssetDisplayName() const override;
+	virtual FLinearColor GetAssetColor() const override;
+	virtual TSoftClassPtr<UObject> GetAssetClass() const override;
+	virtual TConstArrayView<FAssetCategoryPath> GetAssetCategories() const override;
+	// UAssetDefinition End
+};
\ No newline at end of file
diff --git a/Engine/Source/Editor/UMGEditor/Private/AssetTypeActions_WidgetBlueprint.cpp b/Engine/Source/Editor/UMGEditor/Private/AssetTypeActions_WidgetBlueprint.cpp
deleted file mode 100644
diff --git a/Engine/Source/Editor/UMGEditor/Private/AssetTypeActions_WidgetBlueprint.h b/Engine/Source/Editor/UMGEditor/Private/AssetTypeActions_WidgetBlueprint.h
deleted file mode 100644
diff --git a/Engine/Source/Editor/UMGEditor/Private/AssetTypeActions_WidgetBlueprintGeneratedClass.cpp b/Engine/Source/Editor/UMGEditor/Private/AssetTypeActions_WidgetBlueprintGeneratedClass.cpp
deleted file mode 100644
diff --git a/Engine/Source/Editor/UMGEditor/Private/AssetTypeActions_WidgetBlueprintGeneratedClass.h b/Engine/Source/Editor/UMGEditor/Private/AssetTypeActions_WidgetBlueprintGeneratedClass.h
deleted file mode 100644
diff --git a/Engine/Source/Editor/UMGEditor/Private/UMGEditorModule.cpp b/Engine/Source/Editor/UMGEditor/Private/UMGEditorModule.cpp
--- a/Engine/Source/Editor/UMGEditor/Private/UMGEditorModule.cpp
+++ b/Engine/Source/Editor/UMGEditor/Private/UMGEditorModule.cpp
@@ -18,8 +18,6 @@
 #include "AssetToolsModule.h"
 #include "IAssetTypeActions.h"
 #include "AssetTypeActions_SlateVectorArtData.h"
-#include "AssetTypeActions_WidgetBlueprint.h"
-#include "AssetTypeActions_WidgetBlueprintGeneratedClass.h"
 #include "KismetCompilerModule.h"
 #include "WidgetBlueprintCompiler.h"
 
@@ -105,9 +103,7 @@ public:
 		// Register asset types
 		IAssetTools& AssetTools = FModuleManager::LoadModuleChecked<FAssetToolsModule>("AssetTools").Get();
 		RegisterAssetTypeAction(AssetTools, MakeShareable(new FAssetTypeActions_SlateVectorArtData()));
-		RegisterAssetTypeAction(AssetTools, MakeShareable(new FAssetTypeActions_WidgetBlueprint()));
-		RegisterAssetTypeAction(AssetTools, MakeShareable(new FAssetTypeActions_WidgetBlueprintGeneratedClass()));
-
+		
 		FKismetCompilerContext::RegisterCompilerForBP(UWidgetBlueprint::StaticClass(), &UWidgetBlueprint::GetCompilerForWidgetBP );
 
 		// Register with the sequencer module that we provide auto-key handlers.
diff --git a/Engine/Source/Editor/UMGEditor/UMGEditor.Build.cs b/Engine/Source/Editor/UMGEditor/UMGEditor.Build.cs
--- a/Engine/Source/Editor/UMGEditor/UMGEditor.Build.cs
+++ b/Engine/Source/Editor/UMGEditor/UMGEditor.Build.cs
@@ -22,6 +22,7 @@ public class UMGEditor : ModuleRules
 
 		PrivateDependencyModuleNames.AddRange(
 			new string[] {
+				"AssetDefinition",
 				"ClassViewer",
 				"Core",
 				"CoreUObject",
