Remove filters for Blueprint Interface types in the Content Browser and Reference viewer because they don't function properly

diff --git a/Engine/Plugins/Editor/EngineAssetDefinitions/Source/Private/Script/AssetDefinition_Blueprint.cpp b/Engine/Plugins/Editor/EngineAssetDefinitions/Source/Private/Script/AssetDefinition_Blueprint.cpp
--- a/Engine/Plugins/Editor/EngineAssetDefinitions/Source/Private/Script/AssetDefinition_Blueprint.cpp
+++ b/Engine/Plugins/Editor/EngineAssetDefinitions/Source/Private/Script/AssetDefinition_Blueprint.cpp
@@ -120,41 +120,6 @@ UThumbnailInfo* UAssetDefinition_Blueprint::LoadThumbnailInfo(const FAssetData&
 	return UE::Editor::FindOrCreateThumbnailInfo(InAssetData.GetAsset(), USceneThumbnailInfo::StaticClass());
 }
 
-void UAssetDefinition_Blueprint::BuildFilters(TArray<FAssetFilterData>& OutFilters) const
-{
-	Super::BuildFilters(OutFilters);
-
-	{
-		FAssetFilterData Filter;
-		Filter.Name = GetAssetClass().ToSoftObjectPath().ToString() + TEXT("Interface");
-		Filter.DisplayText = LOCTEXT("AssetTypeActions_BlueprintInterface", "Blueprint Interface");
-		Filter.FilterCategories = { EAssetCategoryPaths::Blueprint };
-		Filter.Filter.ClassPaths.Add(GetAssetClass().ToSoftObjectPath().GetAssetPath());
-		Filter.Filter.TagsAndValues.Add(GET_MEMBER_NAME_CHECKED(UBlueprint, BlueprintType), FString(TEXT("BPTYPE_Interface")));
-		OutFilters.Add(MoveTemp(Filter));
-	}
-
-	{
-		FAssetFilterData Filter;
-		Filter.Name = GetAssetClass().ToSoftObjectPath().ToString() + TEXT("MacroLibrary");
-		Filter.DisplayText = LOCTEXT("AssetTypeActions_BlueprintMacroLibrary", "Blueprint Macro Library");
-		Filter.FilterCategories = { EAssetCategoryPaths::Blueprint };
-		Filter.Filter.ClassPaths.Add(GetAssetClass().ToSoftObjectPath().GetAssetPath());
-		Filter.Filter.TagsAndValues.Add(GET_MEMBER_NAME_CHECKED(UBlueprint, BlueprintType), FString(TEXT("BPTYPE_MacroLibrary")));
-		OutFilters.Add(MoveTemp(Filter));
-	}
-
-	{
-		FAssetFilterData Filter;
-		Filter.Name = GetAssetClass().ToSoftObjectPath().ToString() + TEXT("FunctionLibrary");
-		Filter.DisplayText = LOCTEXT("AssetTypeActions_BlueprintFunctionLibrary", "Blueprint Function Library");
-		Filter.Filter.ClassPaths.Add(GetAssetClass().ToSoftObjectPath().GetAssetPath());
-		Filter.FilterCategories = { EAssetCategoryPaths::Blueprint };
-		Filter.Filter.TagsAndValues.Add(GET_MEMBER_NAME_CHECKED(UBlueprint, BlueprintType), FString(TEXT("BPTYPE_FunctionLibrary")));
-		OutFilters.Add(MoveTemp(Filter));
-	}
-}
-
 EAssetCommandResult UAssetDefinition_Blueprint::OpenAssets(const FAssetOpenArgs& OpenArgs) const
 {
 	TArray<FAssetData> OutAssetsThatFailedToLoad;
diff --git a/Engine/Plugins/Editor/EngineAssetDefinitions/Source/Public/Script/AssetDefinition_Blueprint.h b/Engine/Plugins/Editor/EngineAssetDefinitions/Source/Public/Script/AssetDefinition_Blueprint.h
--- a/Engine/Plugins/Editor/EngineAssetDefinitions/Source/Public/Script/AssetDefinition_Blueprint.h
+++ b/Engine/Plugins/Editor/EngineAssetDefinitions/Source/Public/Script/AssetDefinition_Blueprint.h
@@ -36,8 +36,6 @@ public:
 	
 	virtual EAssetCommandResult PerformAssetDiff(const FAssetDiffArgs& DiffArgs) const override;
 
-	virtual void BuildFilters(TArray<FAssetFilterData>& OutFilters) const override;
-	
 	virtual UThumbnailInfo* LoadThumbnailInfo(const FAssetData& InAssetData) const override;
 	virtual EAssetCommandResult OpenAssets(const FAssetOpenArgs& OpenArgs) const override;
 	// UAssetDefinition End
