[Feature] AssetDefinition Merging

diff --git a/Engine/Plugins/Developer/PerforceSourceControl/Source/PerforceSourceControl/Private/PerforceSourceControlState.cpp b/Engine/Plugins/Developer/PerforceSourceControl/Source/PerforceSourceControl/Private/PerforceSourceControlState.cpp
--- a/Engine/Plugins/Developer/PerforceSourceControl/Source/PerforceSourceControl/Private/PerforceSourceControlState.cpp
+++ b/Engine/Plugins/Developer/PerforceSourceControl/Source/PerforceSourceControl/Private/PerforceSourceControlState.cpp
@@ -70,6 +70,11 @@ bool FPerforceSourceControlState::GetOtherBranchHeadModification(FString& HeadBr
 	return !HeadBranchOut.IsEmpty();
 }
 
+FSourceControlChangelistPtr FPerforceSourceControlState::GetCheckInIdentifier() const
+{
+	return MakeShared<FPerforceSourceControlChangelist>(Changelist);
+}
+
 TSharedPtr<class ISourceControlRevision, ESPMode::ThreadSafe> FPerforceSourceControlState::GetBaseRevForMerge() const
 {
 	if (PendingResolveRevNumber == INVALID_REVISION)
diff --git a/Engine/Plugins/Developer/PerforceSourceControl/Source/PerforceSourceControl/Private/PerforceSourceControlState.h b/Engine/Plugins/Developer/PerforceSourceControl/Source/PerforceSourceControl/Private/PerforceSourceControlState.h
--- a/Engine/Plugins/Developer/PerforceSourceControl/Source/PerforceSourceControl/Private/PerforceSourceControlState.h
+++ b/Engine/Plugins/Developer/PerforceSourceControl/Source/PerforceSourceControl/Private/PerforceSourceControlState.h
@@ -86,6 +86,7 @@ public:
 	virtual TArray<FString> GetCheckedOutBranches() const override { return CheckedOutBranches; }
 	virtual FString GetOtherUserBranchCheckedOuts() const override { return OtherUserBranchCheckedOuts; }
 	virtual bool GetOtherBranchHeadModification(FString& HeadBranchOut, FString& ActionOut, int32& HeadChangeListOut) const override;
+	virtual FSourceControlChangelistPtr GetCheckInIdentifier() const override;
 	virtual bool IsCurrent() const override;
 	virtual bool IsSourceControlled() const override;
 	virtual bool IsAdded() const override;
diff --git a/Engine/Plugins/Editor/EngineAssetDefinitions/Source/EngineAssetDefinitions.Build.cs b/Engine/Plugins/Editor/EngineAssetDefinitions/Source/EngineAssetDefinitions.Build.cs
--- a/Engine/Plugins/Editor/EngineAssetDefinitions/Source/EngineAssetDefinitions.Build.cs
+++ b/Engine/Plugins/Editor/EngineAssetDefinitions/Source/EngineAssetDefinitions.Build.cs
@@ -37,6 +37,8 @@ public class EngineAssetDefinitions : ModuleRules
 				"Kismet",
 				"GameProjectGeneration",
 				"AnimationBlueprintEditor",
+				"ToolWidgets",
+                "SourceControl",
 				"Blutility",
 				"UMGEditor",
 				"ToolWidgets"
diff --git a/Engine/Plugins/Editor/EngineAssetDefinitions/Source/Private/AssetDefinition_DataAsset.cpp b/Engine/Plugins/Editor/EngineAssetDefinitions/Source/Private/AssetDefinition_DataAsset.cpp
--- a/Engine/Plugins/Editor/EngineAssetDefinitions/Source/Private/AssetDefinition_DataAsset.cpp
+++ b/Engine/Plugins/Editor/EngineAssetDefinitions/Source/Private/AssetDefinition_DataAsset.cpp
@@ -4,13 +4,19 @@
 
 #include "ClassViewerFilter.h"
 #include "ContentBrowserMenuContexts.h"
+#include "ISourceControlModule.h"
+#include "ISourceControlProvider.h"
 #include "ObjectTools.h"
 #include "SDetailsDiff.h"
+#include "SourceControlHelpers.h"
+#include "SourceControlOperations.h"
 #include "ToolMenu.h"
 #include "ToolMenus.h"
 #include "Kismet2/SClassPickerDialog.h"
 #include "Engine/Engine.h"
 #include "ToolMenuSection.h"
+#include "Serialization/ObjectReader.h"
+#include "Serialization/ObjectWriter.h"
 
 #define LOCTEXT_NAMESPACE "UAssetDefinition_DataAsset"
 
@@ -176,4 +182,295 @@ EAssetCommandResult UAssetDefinition_DataAsset::PerformAssetDiff(const FAssetDif
 	return EAssetCommandResult::Handled;
 }
 
+bool UAssetDefinition_DataAsset::CanMerge() const
+{
+	return true;
+}
+
+struct ScopedMergeResolveTransaction
+{
+	ScopedMergeResolveTransaction(UObject* InManagedObject, EMergeFlags InFlags)
+		: ManagedObject(InManagedObject)
+		, Flags(InFlags)
+	{
+		if (Flags & MF_HANDLE_SOURCE_CONTROL)
+		{
+			UndoHandler = NewObject<UUndoableResolveHandler>();
+			UndoHandler->SetFlags(RF_Transactional);
+			UndoHandler->SetManagedObject(ManagedObject);
+			
+			TransactionNum = GEditor->BeginTransaction(LOCTEXT("ResolveMerge", "ResolveAutoMerge"));
+			ensure(UndoHandler->Modify());
+			ensure(ManagedObject->Modify());
+		}
+	}
+
+	void Cancel()
+	{
+		bCanceled = true;
+	}
+	
+	~ScopedMergeResolveTransaction()
+	{
+		if (Flags & MF_HANDLE_SOURCE_CONTROL)
+		{
+			if (!bCanceled)
+			{
+				UndoHandler->MarkResolved();
+				GEditor->EndTransaction();
+			}
+			else
+			{
+				ManagedObject->GetPackage()->SetDirtyFlag(false);
+				GEditor->CancelTransaction(TransactionNum);
+			}
+		}
+	}
+
+	UObject* ManagedObject;
+	UUndoableResolveHandler* UndoHandler = nullptr;
+	EMergeFlags Flags;
+	int TransactionNum = 0;
+	bool bCanceled = false;
+};
+
+EAssetCommandResult UAssetDefinition_DataAsset::Merge(const FAssetAutomaticMergeArgs& MergeArgs) const
+{
+	FAssetManualMergeArgs ManualMergeArgs;
+	ManualMergeArgs.LocalAsset = MergeArgs.LocalAsset;
+	ManualMergeArgs.ResolutionCallback = MergeArgs.ResolutionCallback;
+	ManualMergeArgs.Flags = MergeArgs.Flags;
+	const UPackage* LocalPackage = ManualMergeArgs.LocalAsset->GetPackage();
+	
+	
+	ISourceControlProvider& SourceControlProvider = ISourceControlModule::Get().GetProvider();
+	const TSharedRef<FUpdateStatus, ESPMode::ThreadSafe> UpdateStatusOperation = ISourceControlOperation::Create<FUpdateStatus>();
+	UpdateStatusOperation->SetUpdateHistory(true);
+	SourceControlProvider.Execute(UpdateStatusOperation, LocalPackage);
+	
+	// Get the SCC state
+	const FSourceControlStatePtr SourceControlState = SourceControlProvider.GetState(LocalPackage, EStateCacheUsage::Use);
+
+	// If we have an asset and its in SCC..
+	if( SourceControlState.IsValid() && SourceControlState->IsSourceControlled() )
+	{
+		const TSharedPtr<ISourceControlRevision, ESPMode::ThreadSafe> RemoteRevision = SourceControlState->GetHistoryItem(0);
+		check(RemoteRevision.IsValid());
+		if(UPackage* TempPackage = DiffUtils::LoadPackageForDiff(RemoteRevision))
+		{
+			// Grab the old asset from that old package
+			ManualMergeArgs.RemoteAsset = FindObject<UObject>(TempPackage, *ManualMergeArgs.LocalAsset->GetName());
+
+			// Recovery for package names that don't match
+			if (ManualMergeArgs.RemoteAsset == nullptr)
+			{
+				ManualMergeArgs.RemoteAsset = TempPackage->FindAssetInPackage();
+			}
+		}
+		
+		const TSharedPtr<ISourceControlRevision, ESPMode::ThreadSafe> BaseRevision = SourceControlState->GetBaseRevForMerge();
+		check(BaseRevision.IsValid());
+		if(UPackage* TempPackage = DiffUtils::LoadPackageForDiff(BaseRevision))
+		{
+			// Grab the old asset from that old package
+			ManualMergeArgs.BaseAsset = FindObject<UObject>(TempPackage, *ManualMergeArgs.LocalAsset->GetName());
+
+			// Recovery for package names that don't match
+			if (ManualMergeArgs.BaseAsset == nullptr)
+			{
+				ManualMergeArgs.BaseAsset = TempPackage->FindAssetInPackage();
+			}
+		}
+	}
+
+	// single asset merging is only supported for assets in a conflicted state in source control
+	if (!ensure(ManualMergeArgs.BaseAsset && ManualMergeArgs.RemoteAsset && ManualMergeArgs.LocalAsset))
+	{
+		return Super::Merge(MergeArgs);
+	}
+	
+	return Merge(ManualMergeArgs);
+}
+
+
+	
+static TArray<uint8> SerializeToBinary(UObject* Object, UObject* Default = nullptr)
+{
+	TArray<uint8> Result;
+	FObjectWriter ObjectWriter(Result);
+	ObjectWriter.SetIsPersistent(true);
+	Default = Default ? Default : Object->GetClass()->ClassDefaultObject;
+	Object->GetClass()->SerializeTaggedProperties( ObjectWriter, reinterpret_cast<uint8*>(Object), Default->GetClass(), reinterpret_cast<uint8*>(Default));
+	return Result;
+};
+
+static void DeserializeFromBinary(UObject* Object, const TArray<uint8>& Data, UObject* Default = nullptr)
+{
+	FObjectReader ObjectReader(Data);
+	ObjectReader.SetIsPersistent(true);
+	Default = Default ? Default : Object->GetClass()->ClassDefaultObject;
+	Object->GetClass()->SerializeTaggedProperties(ObjectReader, reinterpret_cast<uint8*>(Object), Default->GetClass(), reinterpret_cast<uint8*>(Default));
+};
+
+EAssetCommandResult UAssetDefinition_DataAsset::Merge(const FAssetManualMergeArgs& MergeArgs) const
+{
+	auto NotifyResolution = [&MergeArgs](EAssetMergeResult Result)
+	{
+		FAssetMergeResults Results;
+		Results.Result = Result;
+		Results.MergedPackage = MergeArgs.LocalAsset->GetPackage();
+		MergeArgs.ResolutionCallback.ExecuteIfBound(Results);
+		return EAssetCommandResult::Handled;
+	};
+	
+	TArray<UObject*> SubObjects;
+	GetObjectsWithOuter(MergeArgs.LocalAsset, SubObjects);
+	if (!ensure(SubObjects.IsEmpty()))
+	{
+		TArray<FString> SubObjectNames;
+		for (const UObject* SubObject : SubObjects)
+		{
+			SubObjectNames.Add(SubObject->GetName());
+		}
+		UE_LOG(LogEngine, Warning, TEXT("Merge is not currently supported with sub-objects: [{0}]"), *FString::Join(SubObjectNames, TEXT(",")));
+		
+		FAssetMergeResults Results;
+		Results.Result = EAssetMergeResult::Cancelled;
+		Results.MergedPackage = MergeArgs.LocalAsset->GetPackage();
+		MergeArgs.ResolutionCallback.ExecuteIfBound(Results);
+		
+		return NotifyResolution(EAssetMergeResult::Cancelled);
+	}
+
+	auto AutoMerge = [](UObject* HeadRevision, UObject* BranchA, UObject* BranchB, FName Name) -> UObject*
+	{
+		// get delta of BranchB relative to HeadRevision
+		const TArray<uint8> Delta = SerializeToBinary(BranchB, HeadRevision);
+		
+		// apply delta on top of BranchA
+		Name = MakeUniqueObjectName(nullptr, BranchA->GetClass(), Name, EUniqueObjectNameOptions::GloballyUnique);
+		UObject* Merged = DuplicateObject(BranchA, nullptr, Name);
+		DeserializeFromBinary(Merged, Delta, BranchA);
+		return Merged;
+	};
+
+	// apply changes in different orders to come up with two possible merge options
+	UObject* FavorRemote = AutoMerge(MergeArgs.BaseAsset,MergeArgs.LocalAsset, MergeArgs.RemoteAsset, TEXT("FavorRemote"));
+	UObject* FavorLocal = AutoMerge(MergeArgs.BaseAsset,MergeArgs.RemoteAsset,MergeArgs.LocalAsset, TEXT("FavorLocal"));
+	const TArray<uint8> FavorRemoteBinary = SerializeToBinary(FavorRemote, MergeArgs.BaseAsset);
+	const TArray<uint8> FavorLocalBinary = SerializeToBinary(FavorLocal, MergeArgs.BaseAsset);
+
+	// if both merge options are the same, we have no conflicts
+	if (FavorRemoteBinary == FavorLocalBinary)
+	{
+		// copy changes over to the local asset
+		ScopedMergeResolveTransaction(MergeArgs.LocalAsset, MergeArgs.Flags);
+		DeserializeFromBinary(MergeArgs.LocalAsset, FavorLocalBinary);
+		return NotifyResolution(EAssetMergeResult::Completed);
+	}
+	
+	// conflicts detected. We need to ask the user to manually resolve them
+	
+	if (!(MergeArgs.Flags & MF_NO_GUI))
+	{
+		const TSharedRef<SDetailsDiff> DiffView = SDetailsDiff::CreateDiffWindow(FavorRemote, FavorLocal, {}, {}, UDataAsset::StaticClass());
+		DiffView->SetOutputObject(MergeArgs.LocalAsset);
+
+		FObjectReader ObjectReader(FavorLocalBinary);
+		DiffView->RequestModifications(ObjectReader);
+		DiffView->OnWindowClosedEvent.AddLambda([MergeArgs](TSharedRef<SDetailsDiff> DiffView)
+		{
+			// serialize the changes made by the user while the diff view was open
+			TArray<uint8> DiffViewModifications;
+			FObjectWriter ObjectWriter(DiffViewModifications);
+			DiffView->GetModifications(ObjectWriter);
+			
+			// begin undoable transaction
+			ScopedMergeResolveTransaction(MergeArgs.LocalAsset, MergeArgs.Flags);
+			
+			// apply manual changes made within the diff tool
+			DeserializeFromBinary(MergeArgs.LocalAsset, DiffViewModifications);
+			
+			// notify caller that merge is complete
+			FAssetMergeResults Results;
+			Results.Result = EAssetMergeResult::Completed;
+			Results.MergedPackage = MergeArgs.LocalAsset->GetPackage();
+			MergeArgs.ResolutionCallback.ExecuteIfBound(Results);
+		});
+		
+		return EAssetCommandResult::Handled;
+	}
+	
+	return NotifyResolution(EAssetMergeResult::Cancelled);
+}
+
+void UUndoableResolveHandler::SetManagedObject(UObject* Object)
+{
+	ManagedObject = Object;
+	const UPackage* Package = ManagedObject->GetPackage();
+	const FString Filepath = FPaths::ConvertRelativePathToFull(Package->GetLoadedPath().GetLocalFullPath());
+	
+	ISourceControlProvider& Provider = ISourceControlModule::Get().GetProvider();
+	const FSourceControlStatePtr SourceControlState = Provider.GetState(Package, EStateCacheUsage::Use);
+	BaseRevisionNumber = FString::FromInt(SourceControlState->GetBaseRevForMerge()->GetRevisionNumber());
+	CurrentRevisionNumber = FString::FromInt(SourceControlState->GetCurrentRevision()->GetRevisionNumber());
+	CheckinIdentifier = SourceControlState->GetCheckInIdentifier();
+
+	// save package and copy the package to a temp file so it can be reverted
+	const FString BaseFilename = FPaths::GetBaseFilename(Filepath);
+	BackupFilepath = FPaths::CreateTempFilename(*(FPaths::ProjectSavedDir()/TEXT("Temp")), *BaseFilename.Left(32));
+	ensure(FPlatformFileManager::Get().GetPlatformFile().CopyFile(*BackupFilepath, *Filepath));
+}
+
+void UUndoableResolveHandler::MarkResolved()
+{
+	const UPackage* Package = ManagedObject->GetPackage();
+	const FString Filepath = FPaths::ConvertRelativePathToFull(Package->GetLoadedPath().GetLocalFullPath());
+	
+	ISourceControlProvider& Provider = ISourceControlModule::Get().GetProvider();
+	Provider.Execute(ISourceControlOperation::Create<FResolve>(), TArray{Filepath}, EConcurrency::Synchronous);
+	bShouldBeResolved = true;
+}
+
+void UUndoableResolveHandler::PostEditUndo()
+{
+	if (bShouldBeResolved) // undo resolution
+	{
+		MarkResolved();
+	}
+	else // redo resolution
+	{
+		UPackage* Package = ManagedObject->GetPackage();
+		const FString Filepath = FPaths::ConvertRelativePathToFull(Package->GetLoadedPath().GetLocalFullPath());
+		
+		// to force the file to revert to it's pre-resolved state, we must revert, sync back to base revision,
+		// apply the conflicting changes, then sync forward again.
+		ISourceControlProvider& Provider = ISourceControlModule::Get().GetProvider();
+		{
+			const TSharedRef<FSync> SyncOperation = ISourceControlOperation::Create<FSync>();
+			SyncOperation->SetRevision(BaseRevisionNumber);
+			Provider.Execute(SyncOperation, Filepath, EConcurrency::Synchronous);
+		}
+		
+		ResetLoaders(Package);
+		Provider.Execute( ISourceControlOperation::Create<FRevert>(), Filepath, EConcurrency::Synchronous);
+
+		{
+			const TSharedRef<FCheckOut> CheckoutOperation = ISourceControlOperation::Create<FCheckOut>();
+			Provider.Execute(CheckoutOperation, CheckinIdentifier, {Filepath}, EConcurrency::Synchronous);
+		}
+
+		ensure(FPlatformFileManager::Get().GetPlatformFile().CopyFile(*Filepath, *BackupFilepath));
+		
+		{
+			const TSharedRef<FSync> SyncOperation = ISourceControlOperation::Create<FSync>();
+			SyncOperation->SetRevision(CurrentRevisionNumber);
+			Provider.Execute(SyncOperation, Filepath, EConcurrency::Synchronous);
+		}
+
+		Provider.Execute(ISourceControlOperation::Create<FUpdateStatus>(), {Filepath}, EConcurrency::Asynchronous);
+	}
+	UObject::PostEditUndo();
+}
+
 #undef LOCTEXT_NAMESPACE
diff --git a/Engine/Plugins/Editor/EngineAssetDefinitions/Source/Public/AssetDefinition_DataAsset.h b/Engine/Plugins/Editor/EngineAssetDefinitions/Source/Public/AssetDefinition_DataAsset.h
--- a/Engine/Plugins/Editor/EngineAssetDefinitions/Source/Public/AssetDefinition_DataAsset.h
+++ b/Engine/Plugins/Editor/EngineAssetDefinitions/Source/Public/AssetDefinition_DataAsset.h
@@ -24,8 +24,34 @@ public:
 	virtual TSoftClassPtr<UObject> GetAssetClass() const override { return UDataAsset::StaticClass(); }
 	
 	virtual EAssetCommandResult PerformAssetDiff(const FAssetDiffArgs& DiffArgs) const override;
+
+	virtual bool CanMerge() const override;
+	virtual EAssetCommandResult Merge(const FAssetAutomaticMergeArgs& MergeArgs) const override;
+	virtual EAssetCommandResult Merge(const FAssetManualMergeArgs& MergeArgs) const override;
 };
 
 #if UE_ENABLE_INCLUDE_ORDER_DEPRECATED_IN_5_2
 #include "CoreMinimal.h"
 #endif
+
+
+UCLASS()
+class UUndoableResolveHandler : public UObject
+{
+public:
+	GENERATED_BODY()
+	void SetManagedObject(UObject* Object);
+	void MarkResolved();
+
+	virtual void PostEditUndo() override;
+
+private:
+	FString BaseRevisionNumber;
+	FString CurrentRevisionNumber;
+	FString BackupFilepath;
+	UObject* ManagedObject;
+	TSharedPtr<class ISourceControlChangelist> CheckinIdentifier;
+	
+	UPROPERTY()
+	bool bShouldBeResolved = false;
+};
\ No newline at end of file
diff --git a/Engine/Plugins/Experimental/Animation/CurveExpression/Source/Editor/Private/AssetDefinition_CurveExpressionsDataAsset.cpp b/Engine/Plugins/Experimental/Animation/CurveExpression/Source/Editor/Private/AssetDefinition_CurveExpressionsDataAsset.cpp
--- a/Engine/Plugins/Experimental/Animation/CurveExpression/Source/Editor/Private/AssetDefinition_CurveExpressionsDataAsset.cpp
+++ b/Engine/Plugins/Experimental/Animation/CurveExpression/Source/Editor/Private/AssetDefinition_CurveExpressionsDataAsset.cpp
@@ -2,6 +2,7 @@
 
 #include "AssetDefinition_CurveExpressionsDataAsset.h"
 
+#include "AssetDefinitionRegistry.h"
 #include "SDetailsDiff.h"
 
 
@@ -15,3 +16,27 @@ EAssetCommandResult UAssetDefinition_CurveExpressionsDataAsset::PerformAssetDiff
 	SDetailsDiff::CreateDiffWindow(DiffArgs.OldAsset, DiffArgs.NewAsset, DiffArgs.OldRevision, DiffArgs.NewRevision, UCurveExpressionsDataAsset::StaticClass());
 	return EAssetCommandResult::Handled;
 }
+
+bool UAssetDefinition_CurveExpressionsDataAsset::CanMerge() const
+{
+	// deffer to UDataAsset merge methods
+	const UAssetDefinition* AssetDefinition = UAssetDefinitionRegistry::Get()->GetAssetDefinitionForClass(UDataAsset::StaticClass());
+	check(AssetDefinition);
+	return AssetDefinition->CanMerge();
+}
+
+EAssetCommandResult UAssetDefinition_CurveExpressionsDataAsset::Merge(const FAssetAutomaticMergeArgs& MergeArgs) const
+{
+	// deffer to UDataAsset merge methods
+	const UAssetDefinition* AssetDefinition = UAssetDefinitionRegistry::Get()->GetAssetDefinitionForClass(UDataAsset::StaticClass());
+	check(AssetDefinition);
+	return AssetDefinition->Merge(MergeArgs);
+}
+
+EAssetCommandResult UAssetDefinition_CurveExpressionsDataAsset::Merge(const FAssetManualMergeArgs& MergeArgs) const
+{
+	// deffer to UDataAsset merge methods
+	const UAssetDefinition* AssetDefinition = UAssetDefinitionRegistry::Get()->GetAssetDefinitionForClass(UDataAsset::StaticClass());
+	check(AssetDefinition);
+	return AssetDefinition->Merge(MergeArgs);
+}
diff --git a/Engine/Plugins/Experimental/Animation/CurveExpression/Source/Editor/Private/AssetDefinition_CurveExpressionsDataAsset.h b/Engine/Plugins/Experimental/Animation/CurveExpression/Source/Editor/Private/AssetDefinition_CurveExpressionsDataAsset.h
--- a/Engine/Plugins/Experimental/Animation/CurveExpression/Source/Editor/Private/AssetDefinition_CurveExpressionsDataAsset.h
+++ b/Engine/Plugins/Experimental/Animation/CurveExpression/Source/Editor/Private/AssetDefinition_CurveExpressionsDataAsset.h
@@ -24,4 +24,7 @@ public:
 		return Categories;
 	}
 	virtual EAssetCommandResult PerformAssetDiff(const FAssetDiffArgs& DiffArgs) const override;
+	virtual bool CanMerge() const override;
+	virtual EAssetCommandResult Merge(const FAssetAutomaticMergeArgs& MergeArgs) const override;
+	virtual EAssetCommandResult Merge(const FAssetManualMergeArgs& MergeArgs) const override;
 };
diff --git a/Engine/Source/Developer/SourceControl/Public/ISourceControlState.h b/Engine/Source/Developer/SourceControl/Public/ISourceControlState.h
--- a/Engine/Source/Developer/SourceControl/Public/ISourceControlState.h
+++ b/Engine/Source/Developer/SourceControl/Public/ISourceControlState.h
@@ -11,6 +11,11 @@ class ISourceControlState;
 typedef TSharedRef<class ISourceControlState, ESPMode::ThreadSafe> FSourceControlStateRef;
 typedef TSharedPtr<class ISourceControlState, ESPMode::ThreadSafe> FSourceControlStatePtr;
 
+class ISourceControlChangelist;
+
+typedef TSharedRef<class ISourceControlChangelist, ESPMode::ThreadSafe> FSourceControlChangelistRef;
+typedef TSharedPtr<class ISourceControlChangelist, ESPMode::ThreadSafe> FSourceControlChangelistPtr;
+
 /**
  * An abstraction of the state of a file under source control
  */
@@ -140,6 +145,8 @@ public:
 	*/
 	virtual bool GetOtherBranchHeadModification(FString& HeadBranchOut, FString& ActionOut, int32& HeadChangeListOut) const = 0;
 
+	virtual FSourceControlChangelistPtr GetCheckInIdentifier() const;
+
 	/** Get whether this file is up-to-date with the version in source control */
 	virtual bool IsCurrent() const = 0;
 
@@ -211,3 +218,8 @@ public:
 	}
 };
 
+inline FSourceControlChangelistPtr ISourceControlState::GetCheckInIdentifier() const
+{
+	return {};
+}
+
diff --git a/Engine/Source/Editor/AssetDefinition/Public/AssetDefinition.h b/Engine/Source/Editor/AssetDefinition/Public/AssetDefinition.h
--- a/Engine/Source/Editor/AssetDefinition/Public/AssetDefinition.h
+++ b/Engine/Source/Editor/AssetDefinition/Public/AssetDefinition.h
@@ -130,10 +130,18 @@ struct FAssetMergeResults
 
 DECLARE_DELEGATE_OneParam(FOnAssetMergeResolved, const FAssetMergeResults& Results);
 
+enum EMergeFlags : uint8
+{
+	MF_NONE                    = 0x00,
+	MF_NO_GUI                  = 0x01,
+	MF_HANDLE_SOURCE_CONTROL   = 0x02,
+};
+
 struct FAssetAutomaticMergeArgs
 {
 	UObject* LocalAsset = nullptr;
 	FOnAssetMergeResolved ResolutionCallback;
+	EMergeFlags Flags = MF_HANDLE_SOURCE_CONTROL;
 };
 
 struct FAssetManualMergeArgs
@@ -142,6 +150,7 @@ struct FAssetManualMergeArgs
 	UObject* BaseAsset = nullptr;
 	UObject* RemoteAsset = nullptr;
 	FOnAssetMergeResolved ResolutionCallback;
+	EMergeFlags Flags = MF_HANDLE_SOURCE_CONTROL;
 };
 
 struct FAssetSupportResponse
diff --git a/Engine/Source/Editor/Kismet/Private/SDetailsDiff.cpp b/Engine/Source/Editor/Kismet/Private/SDetailsDiff.cpp
--- a/Engine/Source/Editor/Kismet/Private/SDetailsDiff.cpp
+++ b/Engine/Source/Editor/Kismet/Private/SDetailsDiff.cpp
@@ -216,26 +216,28 @@ TSharedRef<SWidget> SDetailsDiff::DefaultEmptyPanel()
 		];
 }
 
-TSharedPtr<SWindow> SDetailsDiff::CreateDiffWindow(FText WindowTitle, const UObject* OldObject, const UObject* NewObject, const FRevisionInfo& OldRevision, const FRevisionInfo& NewRevision)
+TSharedRef<SDetailsDiff> SDetailsDiff::CreateDiffWindow(FText WindowTitle, const UObject* OldObject, const UObject* NewObject, const FRevisionInfo& OldRevision, const FRevisionInfo& NewRevision)
 {
 	// sometimes we're comparing different revisions of one single asset (other 
 	// times we're comparing two completely separate assets altogether)
-	bool bIsSingleAsset = !NewObject || !OldObject || (NewObject->GetName() == OldObject->GetName());
+	const bool bIsSingleAsset = !NewObject || !OldObject || (NewObject->GetName() == OldObject->GetName());
 
 	TSharedPtr<SWindow> Window = SNew(SWindow)
 		.Title(WindowTitle)
 		.ClientSize(FVector2D(1000.f, 800.f));
 
-	Window->SetContent(SNew(SDetailsDiff)
+	TSharedRef<SDetailsDiff> DetailsDiff = SNew(SDetailsDiff)
 		.AssetOld(OldObject)
 		.AssetNew(NewObject)
 		.OldRevision(OldRevision)
 		.NewRevision(NewRevision)
 		.ShowAssetNames(!bIsSingleAsset)
-		.ParentWindow(Window));
+		.ParentWindow(Window);
+	
+	Window->SetContent(DetailsDiff);
 
 	// Make this window a child of the modal window if we've been spawned while one is active.
-	const TSharedPtr<SWindow> ActiveModal = FSlateApplication::Get().GetActiveModalWindow();
+	const TSharedPtr<SWindow> ActiveModal = FSlateApplication::Get().GetActiveTopLevelWindow();
 	if (ActiveModal.IsValid())
 	{
 		FSlateApplication::Get().AddWindowAsNativeChild(Window.ToSharedRef(), ActiveModal.ToSharedRef());
@@ -245,10 +247,19 @@ TSharedPtr<SWindow> SDetailsDiff::CreateDiffWindow(FText WindowTitle, const UObj
 		FSlateApplication::Get().AddWindow(Window.ToSharedRef());
 	}
 
-	return Window;
+	TWeakPtr<SDetailsDiff> SelfWeak = DetailsDiff.ToWeakPtr();
+	Window->SetOnWindowClosed(::FOnWindowClosed::CreateLambda([SelfWeak](const TSharedRef<SWindow>&)
+	{
+		if (const TSharedPtr<SDetailsDiff> Self = SelfWeak.Pin())
+		{
+			Self->OnWindowClosedEvent.Broadcast(Self.ToSharedRef());
+		}
+	}));
+
+	return DetailsDiff;
 }
 
-TSharedPtr<SWindow> SDetailsDiff::CreateDiffWindow(const UObject* OldObject, const UObject* NewObject,
+TSharedRef<SDetailsDiff> SDetailsDiff::CreateDiffWindow(const UObject* OldObject, const UObject* NewObject,
                                                    const FRevisionInfo& OldRevision, const FRevisionInfo& NewRevision, const UClass* ObjectClass)
 {
 	check(OldObject || NewObject);
@@ -269,6 +280,34 @@ TSharedPtr<SWindow> SDetailsDiff::CreateDiffWindow(const UObject* OldObject, con
 	return CreateDiffWindow(WindowTitle, OldObject, NewObject, OldRevision, NewRevision);
 }
 
+void SDetailsDiff::SetOutputObject(const UObject* InOutputObject)
+{
+	OutputObjectUnmodified = InOutputObject;
+	OutputObjectModified = InOutputObject ? DuplicateObject(InOutputObject, nullptr) : nullptr;
+	OnOutputObjectSetEvent.Broadcast();
+}
+
+void SDetailsDiff::GetModifications(FArchive& Archive) const
+{
+	OutputObjectModified->GetClass()->SerializeTaggedProperties(
+		Archive, reinterpret_cast<uint8*>(OutputObjectModified),
+		OutputObjectUnmodified->GetClass(), (uint8*)OutputObjectUnmodified
+	);
+}
+
+void SDetailsDiff::RequestModifications(FArchive& Archive) const
+{
+	OutputObjectModified->GetClass()->SerializeTaggedProperties(
+		Archive, reinterpret_cast<uint8*>(OutputObjectModified),
+		OutputObjectModified->GetClass(), reinterpret_cast<uint8*>(OutputObjectModified->GetClass()->ClassDefaultObject)
+	);
+}
+
+bool SDetailsDiff::IsOutputEnabled() const
+{
+	return OutputObjectUnmodified && OutputObjectModified;
+}
+
 void SDetailsDiff::NextDiff()
 {
 	DiffTreeView::HighlightNextDifference(DifferencesTreeView.ToSharedRef(), RealDifferences, PrimaryDifferencesList);
@@ -313,10 +352,7 @@ SDetailsDiff::FDiffControl SDetailsDiff::GenerateDetailsPanel()
 	const TSharedPtr<FDetailsDiffControl> NewDiffControl = MakeShared<FDetailsDiffControl>(PanelOld.Object, PanelNew.Object, FOnDiffEntryFocused::CreateRaw(this, &SDetailsDiff::SetCurrentMode, DetailsMode), true);
 	NewDiffControl->EnableComments(DifferencesTreeView.ToWeakPtr());
 	NewDiffControl->GenerateTreeEntries(PrimaryDifferencesList, RealDifferences);
-
 	
-	TWeakPtr<SDetailsSplitter> DetailsSplitter;
-
 	const TSharedRef<SDetailsSplitter> Splitter = SNew(SDetailsSplitter);
 	if (PanelOld.Object)
 	{
@@ -337,6 +373,32 @@ SDetailsDiff::FDiffControl SDetailsDiff::GenerateDetailsPanel()
 		);
 	}
 
+	
+	const TWeakPtr<SDetailsSplitter> WeakSplitter = Splitter;
+	const TWeakPtr<FDetailsDiffControl> WeakDiffControl = NewDiffControl;
+	OnOutputObjectSetEvent.AddLambda([WeakSplitter, WeakDiffControl, this]()
+	{
+		const TSharedPtr<SDetailsSplitter> Splitter = WeakSplitter.Pin();
+		const TSharedPtr<FDetailsDiffControl> DiffControl = WeakDiffControl.Pin();
+		if (Splitter && DiffControl)
+		{
+			TSharedRef<IDetailsView> DetailsView = DiffControl->InsertObject(OutputObjectModified, false, 1);
+
+			// allow user to edit the output panel
+			DetailsView->SetIsPropertyEditingEnabledDelegate(FIsPropertyEditingEnabled::CreateStatic([]{return true; }));
+			
+			// insert the output object as a central panel
+			Splitter->AddSlot(
+				SDetailsSplitter::Slot()
+					.DetailsView(DiffControl->GetDetailsWidget(OutputObjectModified))
+					.Value(0.5f)
+					.IsReadonly(false)
+					.DifferencesWithRightPanel(DiffControl.ToSharedRef(), &FDetailsDiffControl::GetDifferencesWithRight, (const UObject*)OutputObjectModified),
+				1 // insert between left and right panel (index 1)
+			);
+		}
+	});
+
 	SDetailsDiff::FDiffControl Ret;
 	Ret.DiffControl = NewDiffControl;
 	Ret.Widget = Splitter;
diff --git a/Engine/Source/Editor/Kismet/Public/SDetailsDiff.h b/Engine/Source/Editor/Kismet/Public/SDetailsDiff.h
--- a/Engine/Source/Editor/Kismet/Public/SDetailsDiff.h
+++ b/Engine/Source/Editor/Kismet/Public/SDetailsDiff.h
@@ -59,11 +59,25 @@ public:
 	static TSharedRef<SWidget> DefaultEmptyPanel();
 
 	/** Helper function to create a window that holds a diff widget */
-	static TSharedPtr<SWindow> CreateDiffWindow(FText WindowTitle, const UObject* OldObject, const UObject* NewObject, const struct FRevisionInfo& OldRevision, const struct FRevisionInfo& NewRevision);
+	static TSharedRef<SDetailsDiff> CreateDiffWindow(FText WindowTitle, const UObject* OldObject, const UObject* NewObject, const struct FRevisionInfo& OldRevision, const struct FRevisionInfo& NewRevision);
 
 	/** Helper function to create a window that holds a diff widget (default window title) */
-	static TSharedPtr<SWindow> CreateDiffWindow(const UObject* OldObject, const UObject* NewObject, const FRevisionInfo& OldRevision, const FRevisionInfo& NewRevision, const UClass* ObjectClass);
+	static TSharedRef<SDetailsDiff> CreateDiffWindow(const UObject* OldObject, const UObject* NewObject, const FRevisionInfo& OldRevision, const FRevisionInfo& NewRevision, const UClass* ObjectClass);
 
+	/** Enables actions like "Choose Left" and "Choose Right" which will modify the OutputObject */
+	void SetOutputObject(const UObject* OutputObject);
+	
+	/** Return a serialized buffer of change requests made by the user */
+	void GetModifications(FArchive& Archive) const;
+	
+	/** submit an archive with change requests to make to OutputObject */
+	void RequestModifications(FArchive& Archive) const;
+
+	/** Returns whether SetOutputObject was called with a valid object */
+	bool IsOutputEnabled() const;
+
+	DECLARE_MULTICAST_DELEGATE_OneParam(FOnWindowClosedEvent, TSharedRef<SDetailsDiff>)
+	FOnWindowClosedEvent OnWindowClosedEvent;
 protected:
 	/** Called when user clicks button to go to next difference */
 	void NextDiff();
@@ -110,8 +124,12 @@ protected:
 	FDetailsDiffPanel PanelOld;
 	FDetailsDiffPanel PanelNew;
 	
-	/** If the two views should be locked */
-	bool bLockViews;
+	/** If set, actions like "Choose Left" and "Choose Right" will be enabled and modify the OutputObject */
+	const UObject* OutputObjectUnmodified = nullptr;
+	UObject* OutputObjectModified = nullptr; // clone that can be mutated by user actions
+
+	DECLARE_MULTICAST_DELEGATE(FOnSetOutputObjectEvent)
+	FOnSetOutputObjectEvent OnOutputObjectSetEvent;
 
 	/** Contents widget that we swap when mode changes (defaults, components, etc) */
 	TSharedPtr<SBox> ModeContents;
diff --git a/Engine/Source/Editor/PropertyEditor/Private/SDetailsView.cpp b/Engine/Source/Editor/PropertyEditor/Private/SDetailsView.cpp
--- a/Engine/Source/Editor/PropertyEditor/Private/SDetailsView.cpp
+++ b/Engine/Source/Editor/PropertyEditor/Private/SDetailsView.cpp
@@ -393,12 +393,28 @@ void SDetailsView::Construct(const FArguments& InArgs, const FDetailsViewArgs& I
 		SNew(SOverlay)
 		+ SOverlay::Slot()
 		[
-			ConstructTreeView(ExternalScrollbar)
+			SNew(SHorizontalBox)
+			+SHorizontalBox::Slot()
+			.HAlign(HAlign_Left)
+			.AutoWidth()
+			[
+				SNew(SBox)
+				.Visibility(DetailsViewArgs.ScrollbarAlignment == HAlign_Left ? EVisibility::Visible : EVisibility::Collapsed)
+				.WidthOverride(16.0f)
+				[
+					ExternalScrollbar
+				]
+			]
+			+SHorizontalBox::Slot()
+			[
+				ConstructTreeView(ExternalScrollbar)
+			]
 		]
 		+ SOverlay::Slot()
-		.HAlign(HAlign_Right)
+		.HAlign(DetailsViewArgs.ScrollbarAlignment)
 		[
 			SNew(SBox)
+			.Visibility(DetailsViewArgs.ScrollbarAlignment == HAlign_Left ? EVisibility::Collapsed : EVisibility::Visible)
 			.WidthOverride(16.0f)
 			[
 				ExternalScrollbar
diff --git a/Engine/Source/Editor/UnrealEd/Private/EditorCommandLineUtils.cpp b/Engine/Source/Editor/UnrealEd/Private/EditorCommandLineUtils.cpp
--- a/Engine/Source/Editor/UnrealEd/Private/EditorCommandLineUtils.cpp
+++ b/Engine/Source/Editor/UnrealEd/Private/EditorCommandLineUtils.cpp
@@ -1,6 +1,8 @@
 // Copyright Epic Games, Inc. All Rights Reserved.
 
 #include "EditorCommandLineUtils.h"
+
+#include "AssetDefinitionRegistry.h"
 #include "HAL/FileManager.h"
 #include "Misc/Paths.h"
 #include "Stats/Stats.h"
@@ -21,6 +23,7 @@
 #include "IAssetTypeActions.h"
 #include "AssetToolsModule.h"
 #include "ProjectDescriptor.h"
+#include "UObject/SavePackage.h"
 
 #define LOCTEXT_NAMESPACE "EditorCommandLineUtils"
 
@@ -235,6 +238,8 @@ public:
 	/**  */
 	const FString& GetAssetFilePath() const;
 
+	UPackage* GetPackage() const;
+
 private:
 	UPackage* Package;
 	UObject*  AssetObj;
@@ -347,6 +352,11 @@ const FString& FMergeAsset::GetAssetFilePath() const
 	return DestFilePath;
 }
 
+UPackage* FMergeAsset::GetPackage() const
+{
+	return Package;
+}
+
 /*******************************************************************************
  * EditorCommandLineUtilsImpl Implementation
  ******************************************************************************/
@@ -608,84 +618,78 @@ static void EditorCommandLineUtilsImpl::RunAssetMerge(FMergeAsset const& Base, F
 	class FMergeResolutionHandler : public TSharedFromThis<FMergeResolutionHandler>
 	{
 	public:
-		FMergeResolutionHandler(UPackage* MergingPkgIn, const FString& DstFilePathIn)
-			: MergingPackage(MergingPkgIn)
-			, Resolution(EMergeResult::Unknown)
+		FMergeResolutionHandler(FMergeAsset const& LocalIn, const FString& DstFilePathIn)
+			: MergingAsset(LocalIn)
+			, Resolution(EAssetMergeResult::Unknown)
 			, DstFilePath(DstFilePathIn)
 		{
 			// force the user to save the result file (so we know if they "accepted" the merge)
-			MergingPkgIn->SetDirtyFlag(true);
+			MergingAsset.GetPackage()->SetDirtyFlag(true);
 		}
 
 		/** Records the user's selected resolution, and closes the editor. */
-		void HandleMergeResolution(UPackage* MergedPackageIn, EMergeResult::Type ResolutionIn)
+		void HandleMergeResolution(UPackage* MergedPackageIn, EAssetMergeResult ResolutionIn)
 		{
-			if (MergedPackageIn == MergingPackage)
+			if (MergedPackageIn == MergingAsset.GetPackage())
 			{
-				if (ResolutionIn == EMergeResult::Cancelled)
+				if (ResolutionIn == EAssetMergeResult::Cancelled)
 				{
 					// they don't want to save any changes, so clear the flag
-					MergingPackage->SetDirtyFlag(false);
+					MergingAsset.GetPackage()->SetDirtyFlag(false);
 				}
 				
-				if (Resolution == EMergeResult::Unknown)
+				if (Resolution == EAssetMergeResult::Unknown)
 				{
 					Resolution = ResolutionIn;
 					EditorCommandLineUtilsImpl::ForceCloseEditor();
 				}
-			}
-		}
 
-		/** Copies the modified file if the user saved changes (and didn't cancel). */
-		void HandleEditorClose()
-		{
-			if ((Resolution != EMergeResult::Cancelled) && !MergingPackage->IsDirty())
-			{
-				FString SrcFilePath = MergingPackage->GetLoadedPath().GetLocalFullPath();
-				IFileManager::Get().Copy(*DstFilePath, *SrcFilePath);
+				if (ResolutionIn == EAssetMergeResult::Completed)
+				{
+					FSavePackageArgs SaveArgs;
+					SaveArgs.TopLevelFlags = RF_Standalone;
+					SaveArgs.Error = GLog;
+					UPackage::SavePackage(MergingAsset.GetPackage(), MergingAsset.GetAssetObj(), *DstFilePath, SaveArgs);
+				}
 			}
 		}
 
 	private:
-		UPackage* MergingPackage;
-		EMergeResult::Type Resolution;
+		FMergeAsset MergingAsset;
+		EAssetMergeResult Resolution;
 		FString DstFilePath;
 	};
 	
-	UPackage* MergeResultPkg = Local.GetAssetObj()->GetOutermost();
 	const FString& ResultFilePath = (!Result.GetSourceFilePath().IsEmpty()) ? Result.GetSourceFilePath() : Local.GetSourceFilePath();
-	TSharedRef<FMergeResolutionHandler> MergeHandler = MakeShareable(new FMergeResolutionHandler(MergeResultPkg, ResultFilePath));
+	TSharedRef<FMergeResolutionHandler> MergeHandler = MakeShareable(new FMergeResolutionHandler(Local, ResultFilePath));
 
 	// we use a lambda delegate to route the call into MergeHandler (we require 
 	// this intermediate to hold onto a MergeHandler ref, so it doesn't get 
 	// prematurely destroyed at the end of this function)
-	auto HandleMergeResolution = [MergeHandler](UPackage* MergedPackageIn, EMergeResult::Type ResolutionIn)
+	auto HandleMergeResolution = [MergeHandler](const FAssetMergeResults& Results)
 	{
-		MergeHandler->HandleMergeResolution(MergedPackageIn, ResolutionIn);
+		MergeHandler->HandleMergeResolution(Results.MergedPackage, Results.Result);
 	};
-	const FOnMergeResolved MergeResolutionDelegate = FOnMergeResolved::CreateLambda(HandleMergeResolution);
+	const FOnAssetMergeResolved MergeResolutionDelegate = FOnAssetMergeResolved::CreateLambda(HandleMergeResolution);
 
 	// have to mount the save directory so that the BP-editor can save
 	// the merged asset packages
 	FPackageName::RegisterMountPoint(TEXT("/Temp/"), FPaths::ProjectSavedDir());
 	
-	IAssetTools& AssetTools = FModuleManager::GetModuleChecked<FAssetToolsModule>("AssetTools").Get();
-	UClass* AssetClass = Local.GetClass();
+	const UClass* AssetClass = Local.GetClass();
 	check(AssetClass != nullptr);
-	TWeakPtr<IAssetTypeActions> AssetActions = AssetTools.GetAssetTypeActionsForClass(AssetClass);
-	check(AssetActions.IsValid());
 	// bring up the merge tool...
-	AssetActions.Pin()->Merge(Base.GetAssetObj(), Remote.GetAssetObj(), Local.GetAssetObj(), MergeResolutionDelegate);
-
-	// we use a lambda delegate to route the call into MergeHandler (we require 
-	// this intermediate to hold onto a MergeHandler ref, so it doesn't get 
-	// prematurely destroyed at the end of this function)
-	auto HandleEditorClose = [](TSharedRef<FMergeResolutionHandler> InMergeHandler)
+	const UAssetDefinition* AssetDefinition = UAssetDefinitionRegistry::Get()->GetAssetDefinitionForClass(AssetClass);
+	if (AssetDefinition->CanMerge())
 	{
-		InMergeHandler->HandleEditorClose();
-	};	
-	// have to copy the file into the expected result file when we're done
-	FEditorDelegates::OnShutdownPostPackagesSaved.AddStatic(HandleEditorClose, MergeHandler);
+		FAssetManualMergeArgs MergeArgs;
+		MergeArgs.BaseAsset = Base.GetAssetObj();
+		MergeArgs.RemoteAsset = Remote.GetAssetObj();
+		MergeArgs.LocalAsset = Local.GetAssetObj();
+		MergeArgs.ResolutionCallback = MergeResolutionDelegate;
+		MergeArgs.Flags = MF_NO_GUI;
+		AssetDefinition->Merge(MergeArgs);
+	}
 }
 
 //------------------------------------------------------------------------------
