[CIS Fix] AssetDefinition_DataAsset.cpp(202): error C2065: 'GEditor': undeclared identifier

diff --git a/Engine/Plugins/Editor/EngineAssetDefinitions/Source/Private/AssetDefinition_DataAsset.cpp b/Engine/Plugins/Editor/EngineAssetDefinitions/Source/Private/AssetDefinition_DataAsset.cpp
--- a/Engine/Plugins/Editor/EngineAssetDefinitions/Source/Private/AssetDefinition_DataAsset.cpp
+++ b/Engine/Plugins/Editor/EngineAssetDefinitions/Source/Private/AssetDefinition_DataAsset.cpp
@@ -2,6 +2,7 @@
 
 #include "AssetDefinition_DataAsset.h"
 
+#include "Editor.h"
 #include "ClassViewerFilter.h"
 #include "ContentBrowserMenuContexts.h"
 #include "ISourceControlModule.h"
