[CIS Fix] AssetDefinition_DataAsset Non-Unity errors

diff --git a/Engine/Plugins/Editor/EngineAssetDefinitions/Source/Private/AssetDefinition_DataAsset.cpp b/Engine/Plugins/Editor/EngineAssetDefinitions/Source/Private/AssetDefinition_DataAsset.cpp
--- a/Engine/Plugins/Editor/EngineAssetDefinitions/Source/Private/AssetDefinition_DataAsset.cpp
+++ b/Engine/Plugins/Editor/EngineAssetDefinitions/Source/Private/AssetDefinition_DataAsset.cpp
@@ -18,6 +18,8 @@
 #include "ToolMenuSection.h"
 #include "Serialization/ObjectReader.h"
 #include "Serialization/ObjectWriter.h"
+#include "HAL/PlatformFileManager.h"
+#include "UObject/Linker.h"
 
 #define LOCTEXT_NAMESPACE "UAssetDefinition_DataAsset"
 
