[CIS Fix] Static analysis missing null check fix

diff --git a/Engine/Plugins/Editor/EngineAssetDefinitions/Source/Private/AssetDefinition_DataAsset.cpp b/Engine/Plugins/Editor/EngineAssetDefinitions/Source/Private/AssetDefinition_DataAsset.cpp
--- a/Engine/Plugins/Editor/EngineAssetDefinitions/Source/Private/AssetDefinition_DataAsset.cpp
+++ b/Engine/Plugins/Editor/EngineAssetDefinitions/Source/Private/AssetDefinition_DataAsset.cpp
@@ -239,6 +239,11 @@ struct ScopedMergeResolveTransaction
 
 EAssetCommandResult UAssetDefinition_DataAsset::Merge(const FAssetAutomaticMergeArgs& MergeArgs) const
 {
+	if (!ensure(MergeArgs.LocalAsset))
+	{
+		return EAssetCommandResult::Unhandled;
+	}
+	
 	FAssetManualMergeArgs ManualMergeArgs;
 	ManualMergeArgs.LocalAsset = MergeArgs.LocalAsset;
 	ManualMergeArgs.ResolutionCallback = MergeArgs.ResolutionCallback;
