Added a context menu option to Animation Blueprints to create a Template
Also integrated Euan Carmichael CL 25133190 to compile and PostUpdate the Animation Blueprint, after assigning new skeleton or creating a template

diff --git a/Engine/Plugins/Editor/EngineAssetDefinitions/Source/EngineAssetDefinitions.Build.cs b/Engine/Plugins/Editor/EngineAssetDefinitions/Source/EngineAssetDefinitions.Build.cs
--- a/Engine/Plugins/Editor/EngineAssetDefinitions/Source/EngineAssetDefinitions.Build.cs
+++ b/Engine/Plugins/Editor/EngineAssetDefinitions/Source/EngineAssetDefinitions.Build.cs
@@ -41,7 +41,8 @@ public class EngineAssetDefinitions : ModuleRules
                 "SourceControl",
 				"Blutility",
 				"UMGEditor",
-				"ToolWidgets"
+				"ToolWidgets",
+				"BlueprintGraph"
 			}
 		);
 		
diff --git a/Engine/Plugins/Editor/EngineAssetDefinitions/Source/Private/Animation/AssetDefinition_AnimBlueprint.cpp b/Engine/Plugins/Editor/EngineAssetDefinitions/Source/Private/Animation/AssetDefinition_AnimBlueprint.cpp
--- a/Engine/Plugins/Editor/EngineAssetDefinitions/Source/Private/Animation/AssetDefinition_AnimBlueprint.cpp
+++ b/Engine/Plugins/Editor/EngineAssetDefinitions/Source/Private/Animation/AssetDefinition_AnimBlueprint.cpp
@@ -22,6 +22,8 @@
 #include "ToolMenuSection.h"
 #include "SPrimaryButton.h"
 #include "SSimpleButton.h"
+#include "BlueprintEditorSettings.h"
+#include "Logging/MessageLog.h"
 
 #define LOCTEXT_NAMESPACE "AssetTypeActions"
 
@@ -212,6 +214,29 @@ namespace MenuExtension_AnimBlueprint
 		}
 	}
 
+	void CompileAnimBlueprintWithLog(UAnimBlueprint* AnimBlueprint)
+	{
+		FCompilerResultsLog	ResultsLog;
+		ResultsLog.SetSourcePath(AnimBlueprint->GetPathName());
+		ResultsLog.BeginEvent(TEXT("Compile"));
+		ResultsLog.bLogDetailedResults = GetDefault<UBlueprintEditorSettings>()->bShowDetailedCompileResults;
+		ResultsLog.EventDisplayThresholdMs = GetDefault<UBlueprintEditorSettings>()->CompileEventDisplayThresholdMs;
+
+		FKismetEditorUtilities::CompileBlueprint(AnimBlueprint, EBlueprintCompileOptions::None, &ResultsLog);
+		ResultsLog.EndEvent();
+		FMessageLog BlueprintLogger("AnimBlueprintLog");
+		BlueprintLogger.NewPage(FText::Format(LOCTEXT("CompileAnimBPPageLabel", "{0}"), FText::FromString(AnimBlueprint->GetName())));
+		BlueprintLogger.AddMessages(ResultsLog.Messages);
+		if ((ResultsLog.NumErrors) > 0)
+		{
+			BlueprintLogger.Notify(FText::Format(LOCTEXT("CompileAnimBPErrors", "AnimBP compilation generated {0} error(s)"), ResultsLog.NumErrors));
+		}
+		else if (ResultsLog.NumWarnings)
+		{
+			BlueprintLogger.Notify(FText::Format(LOCTEXT("CompileAnimBPWarnings", "AnimBP compilation generated {0} warnings(s)"), ResultsLog.NumWarnings));
+		}
+	}
+
 	void ExecuteAssignSkeleton(const FToolMenuContext& InContext)
 	{
 		const UContentBrowserAssetContextMenuContext* CBContext = UContentBrowserAssetContextMenuContext::FindContextWithAssets(InContext);
@@ -239,13 +264,54 @@ namespace MenuExtension_AnimBlueprint
 			{
 				if (AnimBlueprint->TargetSkeleton != SelectedSkeleton)
 				{
+					UAssetEditorSubsystem* AssetEditorSubsystem = GEditor->GetEditorSubsystem<UAssetEditorSubsystem>();
+					const int32 NumEditorsClosed = AssetEditorSubsystem->CloseAllEditorsForAsset(AnimBlueprint);
+
 					AnimBlueprint->Modify();
 					AnimBlueprint->TargetSkeleton = SelectedSkeleton;
+					AnimBlueprint->PostEditChange();
+
+					CompileAnimBlueprintWithLog(AnimBlueprint);
+
+					if (NumEditorsClosed > 0)
+					{
+						AssetEditorSubsystem->OpenEditorForAsset(AnimBlueprint);
+					}
 				}
 			}
 		}
 	}
 
+	void ExecuteCreateTemplate(const FToolMenuContext& InContext, const FAssetData InAnimBlueprintAsset)
+	{
+		if (UAnimBlueprint* SourceBP = Cast<UAnimBlueprint>(InAnimBlueprintAsset.GetAsset()))
+		{
+			UClass* TargetParentClass = SourceBP->GeneratedClass;
+
+			if (!FKismetEditorUtilities::CanCreateBlueprintOfClass(TargetParentClass))
+			{
+				FMessageDialog::Open(EAppMsgType::Ok, LOCTEXT("InvalidClassToMakeBlueprintFrom", "Invalid class with which to make a Blueprint."));
+				return;
+			}
+
+			FString TargetName;
+			FString TargetPackageName;
+			IAssetTools::Get().CreateUniqueAssetName(SourceBP->GetOutermost()->GetName(), TEXT("_Template"), TargetPackageName, TargetName);
+			if (UAnimBlueprint* NewAnimBP = Cast<UAnimBlueprint>(IAssetTools::Get().DuplicateAssetWithDialog(TargetName, TargetPackageName, SourceBP)))
+			{
+				NewAnimBP->Modify();
+				NewAnimBP->bIsTemplate = true;
+				NewAnimBP->TargetSkeleton = nullptr;
+				NewAnimBP->PostEditChange();
+
+				CompileAnimBlueprintWithLog(NewAnimBP);
+
+				UAssetEditorSubsystem* AssetEditorSubsystem = GEditor->GetEditorSubsystem<UAssetEditorSubsystem>();
+				AssetEditorSubsystem->OpenEditorForAsset(NewAnimBP);
+			}
+		}
+	}
+
 	void ExecuteNewSkeletonChildBlueprint(const FToolMenuContext& InContext, const FAssetData InAnimBlueprintAsset)
 	{
 		TArray<FSoftObjectPath> CompatibleSkeletonPaths;
@@ -429,6 +495,25 @@ namespace MenuExtension_AnimBlueprint
 							UIAction.ExecuteAction = FToolMenuExecuteAction::CreateStatic(&ExecuteAssignSkeleton);
 							InSection.AddMenuEntry("AnimBlueprint_AssignSkeleton", Label, ToolTip, Icon, UIAction);
 						}
+
+						{
+							if (const FAssetData* SelectedAnimBlueprintPtr = CBContext->GetSingleSelectedAssetOfType(UAnimBlueprint::StaticClass(), EIncludeSubclasses::No))
+							{
+								// Accept (non-interface) non templated anim BPs or anim BPs
+								if (SelectedAnimBlueprintPtr
+									&& UE::AnimBlueprint::GetBlueprintType(*SelectedAnimBlueprintPtr) != BPTYPE_Interface
+									&& !UE::AnimBlueprint::IsAnimBlueprintTemplate(*SelectedAnimBlueprintPtr))
+								{
+									const TAttribute<FText> Label = LOCTEXT("AnimBlueprint_CreateTemplate", "Create Template Animation Blueprint");
+									const TAttribute<FText> ToolTip = LOCTEXT("AnimBlueprint_CreateTemplateTooltip", "Creates AnimBP Template(s) from the selected Animation Blueprint(s).");
+									const FSlateIcon Icon = FSlateIcon(FAppStyle::GetAppStyleSetName(), "Persona.AssetActions.CreateAnimAsset");
+
+									FToolUIAction UIAction;
+									UIAction.ExecuteAction = FToolMenuExecuteAction::CreateStatic(&ExecuteCreateTemplate, *SelectedAnimBlueprintPtr);
+									InSection.AddMenuEntry("AnimBlueprint_NewSkeletonChildBlueprint", Label, ToolTip, Icon, UIAction);
+								}
+							}
+						}
 					}
 				}
 			}));
