Adapted from PR #9900: Fix FAssetTypeActions_DataTable warnings caused by UE 5.1 FTopLevelAssetPath refactor

diff --git a/Engine/Plugins/Editor/EngineAssetDefinitions/Source/Private/AssetDefinition_DataAsset.cpp b/Engine/Plugins/Editor/EngineAssetDefinitions/Source/Private/AssetDefinition_DataAsset.cpp
--- a/Engine/Plugins/Editor/EngineAssetDefinitions/Source/Private/AssetDefinition_DataAsset.cpp
+++ b/Engine/Plugins/Editor/EngineAssetDefinitions/Source/Private/AssetDefinition_DataAsset.cpp
@@ -164,7 +164,10 @@ FText UAssetDefinition_DataAsset::GetAssetDisplayName(const FAssetData& AssetDat
 		const FAssetDataTagMapSharedView::FFindTagResult RowStructureTag = AssetData.TagsAndValues.FindTag(NAME_RowStructure);
 		if (RowStructureTag.IsSet())
 		{
-			if (UScriptStruct* FoundStruct = UClass::TryFindTypeSlow<UScriptStruct>(RowStructureTag.GetValue(), EFindFirstObjectOptions::ExactClass))
+			// Handle full path names and deprecated short class names
+			const FTopLevelAssetPath ClassPath = FAssetData::TryConvertShortClassNameToPathName(*RowStructureTag.GetValue(), ELogVerbosity::Log);
+
+			if (const UScriptStruct* FoundStruct = UClass::TryFindTypeSlow<UScriptStruct>(ClassPath.ToString(), EFindFirstObjectOptions::ExactClass))
 			{
 				return FText::Format(LOCTEXT("DataTableWithRowType", "Data Table ({0})"), FoundStruct->GetDisplayNameText());
 			}
diff --git a/Engine/Plugins/Editor/EngineAssetDefinitions/Source/Private/Table/AssetDefinition_DataTable.cpp b/Engine/Plugins/Editor/EngineAssetDefinitions/Source/Private/Table/AssetDefinition_DataTable.cpp
--- a/Engine/Plugins/Editor/EngineAssetDefinitions/Source/Private/Table/AssetDefinition_DataTable.cpp
+++ b/Engine/Plugins/Editor/EngineAssetDefinitions/Source/Private/Table/AssetDefinition_DataTable.cpp
@@ -33,7 +33,10 @@ FText UAssetDefinition_DataTable::GetAssetDisplayName(const FAssetData& AssetDat
 		const FAssetDataTagMapSharedView::FFindTagResult RowStructureTag = AssetData.TagsAndValues.FindTag(NAME_RowStructure);
 		if (RowStructureTag.IsSet())
 		{
-			if (UScriptStruct* FoundStruct = UClass::TryFindTypeSlow<UScriptStruct>(RowStructureTag.GetValue(), EFindFirstObjectOptions::ExactClass))
+			// Handle full path names and deprecated short class names
+			const FTopLevelAssetPath ClassPath = FAssetData::TryConvertShortClassNameToPathName(*RowStructureTag.GetValue(), ELogVerbosity::Log);
+
+			if (const UScriptStruct* FoundStruct = UClass::TryFindTypeSlow<UScriptStruct>(ClassPath.ToString(), EFindFirstObjectOptions::ExactClass))
 			{
 				return FText::Format(LOCTEXT("DataTableWithRowType", "Data Table ({0})"), FoundStruct->GetDisplayNameText());
 			}
