[BugFix] Merge stopped working because ISourceControlProvider.GetState doesn't use depot file paths and FResolveInfo::IsValid had a bug

diff --git a/Engine/Plugins/Editor/EngineAssetDefinitions/Source/Private/AssetDefinition_DataAsset.cpp b/Engine/Plugins/Editor/EngineAssetDefinitions/Source/Private/AssetDefinition_DataAsset.cpp
--- a/Engine/Plugins/Editor/EngineAssetDefinitions/Source/Private/AssetDefinition_DataAsset.cpp
+++ b/Engine/Plugins/Editor/EngineAssetDefinitions/Source/Private/AssetDefinition_DataAsset.cpp
@@ -240,6 +240,25 @@ struct ScopedMergeResolveTransaction
 	bool bCanceled = false;
 };
 
+static UPackage* LoadMergePackage(const FString& SCFile, const FString& Revision, const UPackage* LocalPackage)
+{
+	const FString FileWithRevision = SCFile + TEXT("#") + Revision;
+	const TSharedRef<FDownloadFile, ESPMode::ThreadSafe> DownloadFileOperation = ISourceControlOperation::Create<FDownloadFile>(FPaths::DiffDir(), FDownloadFile::EVerbosity::Full);
+	ISourceControlModule::Get().GetProvider().Execute(DownloadFileOperation, FileWithRevision, EConcurrency::Synchronous);
+	const FString DownloadPath = FPaths::ConvertRelativePathToFull(FPaths::DiffDir() / FPaths::GetCleanFilename(FileWithRevision));
+	FString CopyPath = DownloadPath;
+	CopyPath.ReplaceInline(TEXT(".uasset"), TEXT(""));
+	CopyPath.ReplaceCharInline('#', '-');
+	CopyPath.ReplaceCharInline('.', '-');
+	CopyPath = FPaths::CreateTempFilename(*FPaths::GetPath(CopyPath), *FPaths::GetBaseFilename(CopyPath), TEXT(".uasset"));
+
+	if (FPlatformFileManager::Get().GetPlatformFile().CopyFile(*CopyPath, *DownloadPath))
+	{
+		return DiffUtils::LoadPackageForDiff(FPackagePath::FromLocalPath(CopyPath), LocalPackage->GetLoadedPath());
+	}
+	return nullptr;
+}
+
 EAssetCommandResult UAssetDefinition_DataAsset::Merge(const FAssetAutomaticMergeArgs& MergeArgs) const
 {
 	if (!ensure(MergeArgs.LocalAsset))
@@ -267,14 +286,8 @@ EAssetCommandResult UAssetDefinition_DataAsset::Merge(const FAssetAutomaticMerge
 	{
 		const ISourceControlState::FResolveInfo ResolveInfo = SourceControlState->GetResolveInfo();
 		check(ResolveInfo.IsValid());
-
-		SourceControlProvider.Execute(UpdateStatusOperation, ResolveInfo.RemoteFile);
-		const FSourceControlStatePtr RemoteBranch = SourceControlProvider.GetState(ResolveInfo.RemoteFile, EStateCacheUsage::Use);
-		if (!ensure(RemoteBranch.IsValid())) { return EAssetCommandResult::Unhandled; }
-		const TSharedPtr<ISourceControlRevision, ESPMode::ThreadSafe> RemoteRevision = RemoteBranch->FindHistoryRevision(ResolveInfo.RemoteRevision);
-		if (!ensure(RemoteRevision.IsValid())) { return EAssetCommandResult::Unhandled; }
 		
-		if(UPackage* TempPackage = DiffUtils::LoadPackageForDiff(RemoteRevision))
+		if(UPackage* TempPackage = LoadMergePackage(ResolveInfo.RemoteFile, ResolveInfo.RemoteRevision, LocalPackage))
 		{
 			// Grab the old asset from that old package
 			ManualMergeArgs.RemoteAsset = FindObject<UObject>(TempPackage, *ManualMergeArgs.LocalAsset->GetName());
@@ -286,13 +299,7 @@ EAssetCommandResult UAssetDefinition_DataAsset::Merge(const FAssetAutomaticMerge
 			}
 		}
 		
-		SourceControlProvider.Execute(UpdateStatusOperation, ResolveInfo.BaseFile);
-		const FSourceControlStatePtr BaseBranch = SourceControlProvider.GetState(ResolveInfo.BaseFile, EStateCacheUsage::Use);
-		if (!ensure(BaseBranch.IsValid())) { return EAssetCommandResult::Unhandled; }
-		const TSharedPtr<ISourceControlRevision, ESPMode::ThreadSafe> BaseRevision = BaseBranch->FindHistoryRevision(ResolveInfo.BaseRevision);
-		if (!ensure(BaseRevision.IsValid())) { return EAssetCommandResult::Unhandled; }
-		
-		if(UPackage* TempPackage = DiffUtils::LoadPackageForDiff(BaseRevision))
+		if(UPackage* TempPackage = LoadMergePackage(ResolveInfo.BaseFile, ResolveInfo.BaseRevision, LocalPackage))
 		{
 			// Grab the old asset from that old package
 			ManualMergeArgs.BaseAsset = FindObject<UObject>(TempPackage, *ManualMergeArgs.LocalAsset->GetName());
diff --git a/Engine/Source/Developer/SourceControl/Public/ISourceControlState.h b/Engine/Source/Developer/SourceControl/Public/ISourceControlState.h
--- a/Engine/Source/Developer/SourceControl/Public/ISourceControlState.h
+++ b/Engine/Source/Developer/SourceControl/Public/ISourceControlState.h
@@ -33,7 +33,7 @@ public:
 
 		bool IsValid() const
 		{
-			return RemoteRevision.IsEmpty() && !RemoteFile.IsEmpty();
+			return !RemoteRevision.IsEmpty() && !RemoteFile.IsEmpty();
 		}
 
 		operator bool() const
