Fix so anim bp asset definition correctly displays option to create child anim bp

diff --git a/Engine/Plugins/Editor/EngineAssetDefinitions/Source/Private/Animation/AssetDefinition_AnimBlueprint.cpp b/Engine/Plugins/Editor/EngineAssetDefinitions/Source/Private/Animation/AssetDefinition_AnimBlueprint.cpp
--- a/Engine/Plugins/Editor/EngineAssetDefinitions/Source/Private/Animation/AssetDefinition_AnimBlueprint.cpp
+++ b/Engine/Plugins/Editor/EngineAssetDefinitions/Source/Private/Animation/AssetDefinition_AnimBlueprint.cpp
@@ -510,7 +510,7 @@ namespace MenuExtension_AnimBlueprint
 
 									FToolUIAction UIAction;
 									UIAction.ExecuteAction = FToolMenuExecuteAction::CreateStatic(&ExecuteCreateTemplate, *SelectedAnimBlueprintPtr);
-									InSection.AddMenuEntry("AnimBlueprint_NewSkeletonChildBlueprint", Label, ToolTip, Icon, UIAction);
+									InSection.AddMenuEntry("AnimBlueprint_CreateTemplate", Label, ToolTip, Icon, UIAction);
 								}
 							}
 						}
