Added the ability to run find replace on mirror data tables for an existing table.  This functionality is useful if you have new bones or have changed your find  / replace expressions for the table.

diff --git a/Engine/Plugins/Editor/EngineAssetDefinitions/Source/Private/Table/AssetDefinition_MirrorDataTable.cpp b/Engine/Plugins/Editor/EngineAssetDefinitions/Source/Private/Table/AssetDefinition_MirrorDataTable.cpp
new file mode 100644
--- /dev/null
+++ b/Engine/Plugins/Editor/EngineAssetDefinitions/Source/Private/Table/AssetDefinition_MirrorDataTable.cpp
@@ -0,0 +1,53 @@
+// Copyright Epic Games, Inc. All Rights Reserved.
+
+#include "AssetDefinition_MirrorDataTable.h"
+
+#include "AssetDefinitionRegistry.h"
+#include "ScopedTransaction.h"
+#include "ToolMenus.h"
+#include "ContentBrowserMenuContexts.h"
+#include "ToolMenuSection.h"
+
+#define LOCTEXT_NAMESPACE "AssetTypeActions"
+
+// Menu Extensions
+//--------------------------------------------------------------------
+
+namespace MenuExtension_MirrorDataTable
+{
+	void ExecuteFindReplace(const FToolMenuContext& InContext)
+	{
+		const UContentBrowserAssetContextMenuContext* CBContext = UContentBrowserAssetContextMenuContext::FindContextWithAssets(InContext);
+
+		FScopedTransaction Transaction(LOCTEXT("MirrorDataTable_FindReplace", "Update using Find Replace"));
+		for (UMirrorDataTable* MirrorDataTable : CBContext->LoadSelectedObjects<UMirrorDataTable>())
+		{
+			MirrorDataTable->FindReplaceMirroredNames(); 
+		}
+	}
+
+	static FDelayedAutoRegisterHelper DelayedAutoRegister(EDelayedRegisterRunPhase::EndOfEngineInit, [] {
+		UToolMenus::RegisterStartupCallback(FSimpleMulticastDelegate::FDelegate::CreateLambda([]()
+			{
+				FToolMenuOwnerScoped OwnerScoped(UE_MODULE_NAME);
+				UToolMenu* Menu = UE::ContentBrowser::ExtendToolMenu_AssetContextMenu(UMirrorDataTable::StaticClass());
+
+				FToolMenuSection& Section = Menu->FindOrAddSection("GetAssetActions");
+				Section.AddDynamicEntry(NAME_None, FNewToolMenuSectionDelegate::CreateLambda([](FToolMenuSection& InSection)
+					{
+						const UContentBrowserAssetContextMenuContext* CBContext = UContentBrowserAssetContextMenuContext::FindContextWithAssets(InSection);
+						{
+							const TAttribute<FText> Label = LOCTEXT("MirrorDataTable_FindReplace", "Update using Find Replace");
+							const TAttribute<FText> ToolTip = LOCTEXT("MirrorDataTable_FindReplaceTooltip", "Run find replace on all items, adding any new entries based on updated strings");
+							const FSlateIcon Icon = FSlateIcon();
+
+							FToolUIAction UIAction;
+							UIAction.ExecuteAction = FToolMenuExecuteAction::CreateStatic(&ExecuteFindReplace);
+							InSection.AddMenuEntry("MirrorDataTable_FindReplace", Label, ToolTip, Icon, UIAction);
+						}
+					}));
+			}));
+		});
+}
+
+#undef LOCTEXT_NAMESPACE
