Remove AssetVersePath asset tag (motivation is that having a full path written within an asset makes it less portable and changing the root of the verse path involves resaving and revisioning lots of assets)
Bonus: Improve "Create Blueprint/C++ Class" menus to only show if they can be performed (instead of being grayed out)
#rnx

diff --git a/Engine/Plugins/Editor/EngineAssetDefinitions/Source/Private/Script/AssetDefinition_Class.cpp b/Engine/Plugins/Editor/EngineAssetDefinitions/Source/Private/Script/AssetDefinition_Class.cpp
--- a/Engine/Plugins/Editor/EngineAssetDefinitions/Source/Private/Script/AssetDefinition_Class.cpp
+++ b/Engine/Plugins/Editor/EngineAssetDefinitions/Source/Private/Script/AssetDefinition_Class.cpp
@@ -84,130 +84,84 @@ namespace MenuExtension_Class
 			{
 				if (const UContentBrowserAssetContextMenuContext* CBContext = UContentBrowserAssetContextMenuContext::FindContextWithAssets(InSection))
 				{
-					UClass *const BaseClass = (CBContext->SelectedAssets.Num() == 1) ? Cast<UClass>(CBContext->SelectedAssets[0].GetAsset()) : nullptr;
-
-					if (BaseClass)
-					{
-						TArray<UObject::FAssetRegistryTag> InTags;
-						InTags.Reserve(32);
-						BaseClass->GetAssetRegistryTags(InTags);
-						for (const UObject::FAssetRegistryTag& CurrentTag : InTags)
-						{
-							if (CurrentTag.Name == UObject::AssetVersePathTagName())
-							{
-								// Verse devices cannot have C++ or blueprint classes created that derive from them
-								return;
-							}
-						}
-					}
-
-					// Only allow the New class option if we have a base class that we can actually derive from in one of our project modules
-					FGameProjectGenerationModule& GameProjectGenerationModule = FGameProjectGenerationModule::Get();
-					TArray<FModuleContextInfo> ProjectModules = GameProjectGenerationModule.GetCurrentProjectModules();
-					const bool bIsValidBaseCppClass = BaseClass && GameProjectGenerationModule.IsValidBaseClassForCreation(BaseClass, ProjectModules);
-					const bool bIsValidBaseBlueprintClass = BaseClass && FKismetEditorUtilities::CanCreateBlueprintOfClass(BaseClass);
-					auto CreateCreateDerivedCppClass = [BaseClass](const FToolMenuContext& InContext)
-					{
-						// Work out where the header file for the current class is, as we'll use that path as the default for the new class
-						FString BaseClassPath;
-						if(FSourceCodeNavigation::FindClassHeaderPath(BaseClass, BaseClassPath))
-						{
-							// Strip off the actual filename as we only need the path
-							BaseClassPath = FPaths::GetPath(BaseClassPath);
-						}
-
-						FGameProjectGenerationModule::Get().OpenAddCodeToProjectDialog(
-							FAddToProjectConfig()
-							.ParentClass(BaseClass)
-							.InitialPath(BaseClassPath)
-							.ParentWindow(FGlobalTabmanager::Get()->GetRootWindow())
-						);
-					};
+					const FSlateIcon Icon = FSlateIcon(FAppStyle::GetAppStyleSetName(), "Icons.C++");
 
 					auto IsCPPAllowed = []()
-					{
-						return ensure(GUnrealEd) && GUnrealEd->GetUnrealEdOptions()->IsCPPAllowed();
-					};
-
-					auto CanCreateDerivedCppClass = [bIsValidBaseCppClass, IsCPPAllowed](const FToolMenuContext& InContext) -> bool
-					{
-						return bIsValidBaseCppClass && IsCPPAllowed();
-					};
-
-					auto CreateCreateDerivedBlueprintClass = [BaseClass](const FToolMenuContext& InContext)
-					{
-						FGameProjectGenerationModule::Get().OpenAddBlueprintToProjectDialog(
-							FAddToProjectConfig()
-							.ParentClass(BaseClass)
-							.ParentWindow(FGlobalTabmanager::Get()->GetRootWindow())
-						);
-					};
-
-					auto CanCreateDerivedBlueprintClass = [bIsValidBaseBlueprintClass](const FToolMenuContext& InContext) -> bool
-					{
-						return bIsValidBaseBlueprintClass;
-					};
+						{
+							return ensure(GUnrealEd) && GUnrealEd->GetUnrealEdOptions()->IsCPPAllowed();
+						};
 
-					FText NewDerivedCppClassLabel;
-					FText NewDerivedCppClassToolTip;
-					FText NewDerivedBlueprintClassLabel;
-					FText NewDerivedBlueprintClassToolTip;
 					if (CBContext->SelectedAssets.Num() == 1)
 					{
-						check(BaseClass);
+						if (UClass* const BaseClass = Cast<UClass>(CBContext->SelectedAssets[0].GetAsset()))
+						{
+							const FText BaseClassName = FText::FromName(BaseClass->GetFName());
 
-						const FText BaseClassName = FText::FromName(BaseClass->GetFName());
+							if (IsCPPAllowed())
+							{
+								// Only allow the New class option if we have a base class that we can actually derive from in one of our project modules
+								FGameProjectGenerationModule& GameProjectGenerationModule = FGameProjectGenerationModule::Get();
+								TArray<FModuleContextInfo> ProjectModules = GameProjectGenerationModule.GetCurrentProjectModules();
+								if (GameProjectGenerationModule.IsValidBaseClassForCreation(BaseClass, ProjectModules))
+								{
+									const TAttribute<FText> Label = FText::Format(LOCTEXT("Class_NewDerivedCppClassLabel_CreateFrom", "Create C++ class derived from {0}"), BaseClassName);
+									const TAttribute<FText> ToolTip = FText::Format(LOCTEXT("Class_NewDerivedCppClassTooltip_CreateFrom", "Create a new C++ class deriving from {0}."), BaseClassName);
+									FToolUIAction UIAction;
+									UIAction.ExecuteAction = FToolMenuExecuteAction::CreateLambda(
+										[BaseClass](const FToolMenuContext& InContext)
+										{
+											// Work out where the header file for the current class is, as we'll use that path as the default for the new class
+											FString BaseClassPath;
+											if (FSourceCodeNavigation::FindClassHeaderPath(BaseClass, BaseClassPath))
+											{
+												// Strip off the actual filename as we only need the path
+												BaseClassPath = FPaths::GetPath(BaseClassPath);
+											}
+
+											FGameProjectGenerationModule::Get().OpenAddCodeToProjectDialog(
+												FAddToProjectConfig()
+												.ParentClass(BaseClass)
+												.InitialPath(BaseClassPath)
+												.ParentWindow(FGlobalTabmanager::Get()->GetRootWindow()));
+										});
+									InSection.AddMenuEntry("NewDerivedCppClass", Label, ToolTip, Icon, UIAction);
+								}
+							}
 
-						NewDerivedCppClassLabel = FText::Format(LOCTEXT("Class_NewDerivedCppClassLabel_CreateFrom", "Create C++ class derived from {0}"), BaseClassName);
-						if(bIsValidBaseCppClass)
-						{
-							NewDerivedCppClassToolTip = FText::Format(LOCTEXT("Class_NewDerivedCppClassTooltip_CreateFrom", "Create a new C++ class deriving from {0}."), BaseClassName);
-						}
-						else
-						{
-							NewDerivedCppClassToolTip = FText::Format(LOCTEXT("Class_NewDerivedCppClassTooltip_InvalidClass", "Cannot create a new C++ class deriving from {0}."), BaseClassName);
+							if (FKismetEditorUtilities::CanCreateBlueprintOfClass(BaseClass))
+							{
+								const TAttribute<FText> Label = FText::Format(LOCTEXT("Class_NewDerivedBlueprintClassLabel_CreateFrom", "Create Blueprint class based on {0}"), BaseClassName);
+								const TAttribute<FText> ToolTip = FText::Format(LOCTEXT("Class_NewDerivedBlueprintClassTooltip_CreateFrom", "Create a new Blueprint class based on {0}."), BaseClassName);
+								FToolUIAction UIAction;
+								UIAction.ExecuteAction = FToolMenuExecuteAction::CreateLambda(
+									[BaseClass](const FToolMenuContext& InContext)
+									{
+										FGameProjectGenerationModule::Get().OpenAddBlueprintToProjectDialog(
+											FAddToProjectConfig()
+											.ParentClass(BaseClass)
+											.ParentWindow(FGlobalTabmanager::Get()->GetRootWindow()));
+									});
+								InSection.AddMenuEntry("NewDerivedBlueprintClass", Label, ToolTip, Icon, UIAction);
+							}
 						}
+					}
+					else
+					{
+						FToolUIAction CannotExecuteAction;
+						CannotExecuteAction.CanExecuteAction = FToolMenuCanExecuteAction::CreateLambda([](const FToolMenuContext&) { return false; });
 
-						NewDerivedBlueprintClassLabel = FText::Format(LOCTEXT("Class_NewDerivedBlueprintClassLabel_CreateFrom", "Create Blueprint class based on {0}"), BaseClassName);
-						if(bIsValidBaseBlueprintClass)
+						if (IsCPPAllowed())
 						{
-							NewDerivedBlueprintClassToolTip = FText::Format(LOCTEXT("Class_NewDerivedBlueprintClassTooltip_CreateFrom", "Create a new Blueprint class based on {0}."), BaseClassName);
+							const TAttribute<FText> Label = LOCTEXT("Class_NewDerivedCppClassLabel_InvalidNumberOfBases", "New C++ class derived from...");
+							const TAttribute<FText> ToolTip = LOCTEXT("Class_NewDerivedCppClassTooltip_InvalidNumberOfBases", "Can only create a derived C++ class when there is a single base class selected.");
+							InSection.AddMenuEntry("NewDerivedCppClass", Label, ToolTip, Icon, CannotExecuteAction);
 						}
-						else
 						{
-							NewDerivedBlueprintClassToolTip = FText::Format(LOCTEXT("Class_NewDerivedBlueprintClassTooltip_InvalidClass", "Cannot create a new Blueprint class based on {0}."), BaseClassName);
+							const TAttribute<FText> Label = LOCTEXT("Class_NewDerivedBlueprintClassLabel_InvalidNumberOfBases", "New Blueprint class based on...");;
+							const TAttribute<FText> ToolTip = LOCTEXT("Class_NewDerivedBlueprintClassTooltip_InvalidNumberOfBases", "Can only create a Blueprint class when there is a single base class selected.");
+							InSection.AddMenuEntry("NewDerivedBlueprintClass", Label, ToolTip, Icon, CannotExecuteAction);
 						}
 					}
-					else
-					{
-						NewDerivedCppClassLabel = LOCTEXT("Class_NewDerivedCppClassLabel_InvalidNumberOfBases", "New C++ class derived from...");
-						NewDerivedCppClassToolTip = LOCTEXT("Class_NewDerivedCppClassTooltip_InvalidNumberOfBases", "Can only create a derived C++ class when there is a single base class selected.");
-
-						NewDerivedBlueprintClassLabel = LOCTEXT("Class_NewDerivedBlueprintClassLabel_InvalidNumberOfBases", "New Blueprint class based on...");
-						NewDerivedBlueprintClassToolTip = LOCTEXT("Class_NewDerivedBlueprintClassTooltip_InvalidNumberOfBases", "Can only create a Blueprint class when there is a single base class selected.");
-					}
-
-					if (IsCPPAllowed())
-					{
-						const TAttribute<FText> Label = NewDerivedCppClassLabel;
-						const TAttribute<FText> ToolTip = NewDerivedCppClassToolTip;
-						const FSlateIcon Icon = FSlateIcon(FAppStyle::GetAppStyleSetName(), "Icons.C++");
-
-						FToolUIAction UIAction;
-						UIAction.ExecuteAction = FToolMenuExecuteAction::CreateLambda(CreateCreateDerivedCppClass);
-						UIAction.CanExecuteAction = FToolMenuCanExecuteAction::CreateLambda(CanCreateDerivedCppClass);
-						InSection.AddMenuEntry("NewDerivedCppClass", Label, ToolTip, Icon, UIAction);
-					}
-					{
-                    	const TAttribute<FText> Label = NewDerivedBlueprintClassLabel;
-                    	const TAttribute<FText> ToolTip = NewDerivedBlueprintClassToolTip;
-                    	const FSlateIcon Icon = FSlateIcon(FAppStyle::GetAppStyleSetName(), "Icons.C++");
-
-                    	FToolUIAction UIAction;
-                    	UIAction.ExecuteAction = FToolMenuExecuteAction::CreateLambda(CreateCreateDerivedBlueprintClass);
-			            UIAction.CanExecuteAction = FToolMenuCanExecuteAction::CreateLambda(CanCreateDerivedBlueprintClass);
-                    	InSection.AddMenuEntry("NewDerivedBlueprintClass", Label, ToolTip, Icon, UIAction);
-                    }
 				}
 			}));
 		}));
diff --git a/Engine/Source/Editor/ContentBrowser/Private/AssetViewWidgets.cpp b/Engine/Source/Editor/ContentBrowser/Private/AssetViewWidgets.cpp
--- a/Engine/Source/Editor/ContentBrowser/Private/AssetViewWidgets.cpp
+++ b/Engine/Source/Editor/ContentBrowser/Private/AssetViewWidgets.cpp
@@ -816,14 +816,6 @@ TSharedRef<SWidget> SAssetViewItem::CreateToolTipWidget() const
 				AddToToolTipInfoBox(InfoBox, LOCTEXT("TileViewTooltipPath", "Path"), FText::FromName(AssetItem->GetItem().GetVirtualPath()), false);
 			}
 
-			{
-				FString VersePath;
-				if (ItemAssetData.GetTagValue(UObject::AssetVersePathTagName(), VersePath))
-				{
-					AddToToolTipInfoBox( InfoBox, LOCTEXT( "TileViewTooltipVersePath", "Asset Verse Path" ), FText::FromString(VersePath), false);
-				}
-			}
-
 			if (ItemAssetData.IsValid() && ItemAssetData.PackageName != NAME_None)
 			{
 				const FString PackagePathWithinRoot = ContentBrowserUtils::GetPackagePathWithinRoot(ItemAssetData.PackageName.ToString());
diff --git a/Engine/Source/Runtime/CoreUObject/Private/UObject/Obj.cpp b/Engine/Source/Runtime/CoreUObject/Private/UObject/Obj.cpp
--- a/Engine/Source/Runtime/CoreUObject/Private/UObject/Obj.cpp
+++ b/Engine/Source/Runtime/CoreUObject/Private/UObject/Obj.cpp
@@ -2164,11 +2164,6 @@ void UObject::GetAssetRegistryTags(TArray<FAssetRegistryTag>& OutTags) const
 		OutTags.Add(FAssetRegistryTag(FPrimaryAssetId::PrimaryAssetNameTag, PrimaryAssetId.PrimaryAssetName.ToString(), UObject::FAssetRegistryTag::TT_Alphabetical));
 	}
 
-	if (UE::Core::FVersePath VersePath = this->GetVersePath())
-	{
-		OutTags.Emplace(UObject::AssetVersePathTagName(), MoveTemp(VersePath).ToString(), UObject::FAssetRegistryTag::TT_Alphabetical);
-	}
-
 	GetAssetRegistryTagsFromSearchableProperties(this, OutTags);
 
 #if WITH_EDITOR
@@ -2250,12 +2245,6 @@ const FName& UObject::SourceFileTagName()
 	return SourceFilePathName;
 }
 
-const FName& UObject::AssetVersePathTagName()
-{
-	static const FName AssetVersePathTag = TEXT("AssetVersePath");
-	return AssetVersePathTag;
-}
-
 #if WITH_EDITOR
 
 static void PostLoadAssetRegistryTagProperty(FProperty* Prop, const FAssetData& AssetData, TArray<UObject::FAssetRegistryTag>& OutTagsAndValuesToUpdate)
diff --git a/Engine/Source/Runtime/CoreUObject/Public/UObject/Object.h b/Engine/Source/Runtime/CoreUObject/Public/UObject/Object.h
--- a/Engine/Source/Runtime/CoreUObject/Public/UObject/Object.h
+++ b/Engine/Source/Runtime/CoreUObject/Public/UObject/Object.h
@@ -893,9 +893,6 @@ public:
 	/** Get the common tag name used for all asset source file import paths */
 	static COREUOBJECT_API const FName& SourceFileTagName();
 
-	/** Get the common tag name used for all asset verse paths */
-	static COREUOBJECT_API const FName& AssetVersePathTagName();
-
 #if WITH_EDITOR
 
 	/**
