Added CanLocalize() override that return FAssetSupportResponse::NotSupported() for
* UAssetDefinition_HLODLayer
* UAssetDefinition_DataLayer
* UAssetDefinition_DataLayerInstance

diff --git a/Engine/Plugins/Editor/EngineAssetDefinitions/Source/Private/AssetDefinition_DataLayer.h b/Engine/Plugins/Editor/EngineAssetDefinitions/Source/Private/AssetDefinition_DataLayer.h
--- a/Engine/Plugins/Editor/EngineAssetDefinitions/Source/Private/AssetDefinition_DataLayer.h
+++ b/Engine/Plugins/Editor/EngineAssetDefinitions/Source/Private/AssetDefinition_DataLayer.h
@@ -18,6 +18,7 @@ public:
 	virtual FText GetAssetDisplayName() const override { return NSLOCTEXT("AssetTypeActions", "AssetTypeActions_DataLayer", "Data Layer"); }
 	virtual FLinearColor GetAssetColor() const override { return FLinearColor(FColor(12, 65, 12)); }
 	virtual TSoftClassPtr<UObject> GetAssetClass() const override { return UDataLayerAsset::StaticClass(); }
+	virtual FAssetSupportResponse CanLocalize(const FAssetData& InAsset) const override { return FAssetSupportResponse::NotSupported(); }
 	virtual TConstArrayView<FAssetCategoryPath> GetAssetCategories() const override
 	{
 		static const auto Categories = { EAssetCategoryPaths::World };
diff --git a/Engine/Plugins/Editor/EngineAssetDefinitions/Source/Private/AssetDefinition_DataLayerInstance.h b/Engine/Plugins/Editor/EngineAssetDefinitions/Source/Private/AssetDefinition_DataLayerInstance.h
--- a/Engine/Plugins/Editor/EngineAssetDefinitions/Source/Private/AssetDefinition_DataLayerInstance.h
+++ b/Engine/Plugins/Editor/EngineAssetDefinitions/Source/Private/AssetDefinition_DataLayerInstance.h
@@ -19,6 +19,7 @@ public:
 	virtual FText GetAssetDisplayName() const override { return NSLOCTEXT("AssetTypeActions", "AssetTypeActions_DataLayerInstance", "Data Layer Instance"); }
 	virtual FLinearColor GetAssetColor() const override { return FLinearColor(FColor(52, 213, 235)); }
 	virtual TSoftClassPtr<UObject> GetAssetClass() const override { return UDataLayerInstance::StaticClass(); }
+	virtual FAssetSupportResponse CanLocalize(const FAssetData& InAsset) const override { return FAssetSupportResponse::NotSupported(); }
 	virtual TConstArrayView<FAssetCategoryPath> GetAssetCategories() const override
 	{
 		return TConstArrayView<FAssetCategoryPath>();
diff --git a/Engine/Plugins/Editor/EngineAssetDefinitions/Source/Private/AssetDefinition_HLODLayer.h b/Engine/Plugins/Editor/EngineAssetDefinitions/Source/Private/AssetDefinition_HLODLayer.h
--- a/Engine/Plugins/Editor/EngineAssetDefinitions/Source/Private/AssetDefinition_HLODLayer.h
+++ b/Engine/Plugins/Editor/EngineAssetDefinitions/Source/Private/AssetDefinition_HLODLayer.h
@@ -18,6 +18,7 @@ public:
 	virtual FText GetAssetDisplayName() const override { return NSLOCTEXT("AssetTypeActions", "AssetTypeActions_HLODLayer", "HLOD Layer"); }
 	virtual FLinearColor GetAssetColor() const override { return FLinearColor(FColor(0, 200, 200)); }
 	virtual TSoftClassPtr<UObject> GetAssetClass() const override { return UHLODLayer::StaticClass(); }
+	virtual FAssetSupportResponse CanLocalize(const FAssetData& InAsset) const override { return FAssetSupportResponse::NotSupported(); }
 	virtual TConstArrayView<FAssetCategoryPath> GetAssetCategories() const override
 	{
 		static const auto Categories = { EAssetCategoryPaths::World };
