Null guard parent class for BP action menu

#rnx

diff --git a/Engine/Plugins/Editor/EngineAssetDefinitions/Source/Private/Script/AssetDefinition_Blueprint.cpp b/Engine/Plugins/Editor/EngineAssetDefinitions/Source/Private/Script/AssetDefinition_Blueprint.cpp
--- a/Engine/Plugins/Editor/EngineAssetDefinitions/Source/Private/Script/AssetDefinition_Blueprint.cpp
+++ b/Engine/Plugins/Editor/EngineAssetDefinitions/Source/Private/Script/AssetDefinition_Blueprint.cpp
@@ -337,7 +337,7 @@ namespace MenuExtension_Blueprint
 							
 							// Ensure that all the selected blueprints are actors
 							const bool bAreAllSelectedBlueprintsActors =
-								Algo::AllOf(SelectedBlueprintParentClasses, [](UClass* ParentClass){ return ParentClass->IsChildOf(AActor::StaticClass()); });
+								Algo::AllOf(SelectedBlueprintParentClasses, [](UClass* ParentClass){ return ParentClass && ParentClass->IsChildOf(AActor::StaticClass()); });
 							
 							if (bAreAllSelectedBlueprintsActors)
 							{
