Fix the EditorUtilityWidget Blueprint asset definition to match what the Widget Blueprint asset definition does on open for edit.

Both are initializing a Widget Blueprint editor context, which does not support defaults-only mode.

diff --git a/Engine/Plugins/Editor/EngineAssetDefinitions/Source/Private/AssetDefinition_EditorUtilityWidgetBlueprint.cpp b/Engine/Plugins/Editor/EngineAssetDefinitions/Source/Private/AssetDefinition_EditorUtilityWidgetBlueprint.cpp
--- a/Engine/Plugins/Editor/EngineAssetDefinitions/Source/Private/AssetDefinition_EditorUtilityWidgetBlueprint.cpp
+++ b/Engine/Plugins/Editor/EngineAssetDefinitions/Source/Private/AssetDefinition_EditorUtilityWidgetBlueprint.cpp
@@ -51,13 +51,15 @@ EAssetCommandResult UAssetDefinition_EditorUtilityWidgetBlueprint::OpenAssets(co
 
 	for (UBlueprint* Blueprint : OpenArgs.LoadObjects<UBlueprint>())
 	{
-		if (Blueprint->SkeletonGeneratedClass && Blueprint->GeneratedClass)
+		if (Blueprint && Blueprint->SkeletonGeneratedClass && Blueprint->GeneratedClass)
 		{
 			TSharedRef<FWidgetBlueprintEditor> NewBlueprintEditor(new FWidgetBlueprintEditor());
 
+			const bool bShouldOpenInDefaultsMode = false;
 			TArray<UBlueprint*> Blueprints;
 			Blueprints.Add(Blueprint);
-			NewBlueprintEditor->InitWidgetBlueprintEditor(EToolkitMode::Standalone, nullptr, Blueprints, true);
+
+			NewBlueprintEditor->InitWidgetBlueprintEditor(EToolkitMode::Standalone, nullptr, Blueprints, bShouldOpenInDefaultsMode);
 		}
 		else
 		{
