Fix for issue where anim streamable couldn't be created directly from anim sequence asset

#changelist validated
#virtualized

diff --git a/Engine/Plugins/Editor/EngineAssetDefinitions/Source/Private/Animation/AssetDefinition_AnimSequence.cpp b/Engine/Plugins/Editor/EngineAssetDefinitions/Source/Private/Animation/AssetDefinition_AnimSequence.cpp
--- a/Engine/Plugins/Editor/EngineAssetDefinitions/Source/Private/Animation/AssetDefinition_AnimSequence.cpp
+++ b/Engine/Plugins/Editor/EngineAssetDefinitions/Source/Private/Animation/AssetDefinition_AnimSequence.cpp
@@ -69,7 +69,7 @@ namespace MenuExtension_AnimSequence
 		const UContentBrowserAssetContextMenuContext* CBContext = UContentBrowserAssetContextMenuContext::FindContextWithAssets(InContext);
 
 		IAssetTools::Get().CreateAssetsFrom<UAnimSequence>(
-			CBContext->LoadSelectedObjects<UAnimSequence>(), UAnimMontage::StaticClass(), TEXT("_Streamable"), [](UAnimSequence* SourceObject)
+			CBContext->LoadSelectedObjects<UAnimSequence>(), UAnimStreamable::StaticClass(), TEXT("_Streamable"), [](UAnimSequence* SourceObject)
 			{
 				UAnimStreamableFactory* Factory = NewObject<UAnimStreamableFactory>();
 				Factory->SourceAnimation = SourceObject;
