Update redirector references: Change context menu in asset context menu to match folder context menu

diff --git a/Engine/Plugins/Editor/EngineAssetDefinitions/Source/Private/AssetDefinition_Redirector.cpp b/Engine/Plugins/Editor/EngineAssetDefinitions/Source/Private/AssetDefinition_Redirector.cpp
--- a/Engine/Plugins/Editor/EngineAssetDefinitions/Source/Private/AssetDefinition_Redirector.cpp
+++ b/Engine/Plugins/Editor/EngineAssetDefinitions/Source/Private/AssetDefinition_Redirector.cpp
@@ -82,7 +82,7 @@ namespace MenuExtension_Redirector
 					InSection.AddMenuEntry("Redirector_FindTarget", Label, ToolTip, Icon, UIAction);
 				}
 				{
-					const TAttribute<FText> Label = LOCTEXT("Redirector_UpdateReferencers", "Update References");
+					const TAttribute<FText> Label = LOCTEXT("Redirector_UpdateReferencers", "Update Redirector References");
 					const TAttribute<FText> ToolTip = LOCTEXT("Redirector_FixUpTooltip", "Finds references to selected redirectors and resaves the referencing assets if possible, so that they reference the target of the redirector directly instead.");
 					const FSlateIcon Icon = FSlateIcon(FAppStyle::GetAppStyleSetName(), "ClassIcon.ObjectRedirector");
 					const FToolMenuExecuteAction UIAction = FToolMenuExecuteAction::CreateStatic(&ExecuteFixUp);
