Do not allow users to create a child blueprint based off of a BP function library. This never worked, and you would only recieve comipliation errors if you made one.

diff --git a/Engine/Plugins/Editor/EngineAssetDefinitions/Source/Private/Script/AssetDefinition_BlueprintGeneratedClass.cpp b/Engine/Plugins/Editor/EngineAssetDefinitions/Source/Private/Script/AssetDefinition_BlueprintGeneratedClass.cpp
--- a/Engine/Plugins/Editor/EngineAssetDefinitions/Source/Private/Script/AssetDefinition_BlueprintGeneratedClass.cpp
+++ b/Engine/Plugins/Editor/EngineAssetDefinitions/Source/Private/Script/AssetDefinition_BlueprintGeneratedClass.cpp
@@ -118,13 +118,20 @@ namespace MenuExtension_BlueprintGeneratedClass
 							const TAttribute<FText> ToolTip = TAttribute<FText>::CreateLambda([GeneratedClassAsset]()
 							{
 								const uint32 BPFlags = GeneratedClassAsset.GetTagValueRef<uint32>(FBlueprintTags::ClassFlags);
-								if ((BPFlags & (CLASS_Deprecated)) == 0)
+								const EBlueprintType BPType = static_cast<EBlueprintType>(GeneratedClassAsset.GetTagValueRef<int32>(FBlueprintTags::BlueprintType));
+								
+								if ((BPFlags & (CLASS_Deprecated)) != 0)
 								{
-									return LOCTEXT("Blueprint_NewDerivedBlueprintTooltip", "Creates a Child Blueprint Class based on the current Blueprint, allowing you to create variants easily.");
+									return LOCTEXT("Blueprint_NewDerivedBlueprintIsDeprecatedTooltip", "Blueprint class is deprecated, cannot derive a child Blueprint!");
+								}
+								else if (BPType == BPTYPE_FunctionLibrary)
+								{
+									return LOCTEXT("Blueprint_NewDerivedBlueprintIsFunctionLibraryTooltip", "Cannot derive from Blueprint Function Libraries!");
 								}
 								else
 								{
-									return LOCTEXT("Blueprint_NewDerivedBlueprintIsDeprecatedTooltip", "Blueprint class is deprecated, cannot derive a child Blueprint!");
+									return LOCTEXT("Blueprint_NewDerivedBlueprintTooltip", "Creates a Child Blueprint Class based on the current Blueprint, allowing you to create variants easily.");
+									
 								}
 							});
 							const FSlateIcon Icon = FSlateIcon(FAppStyle::GetAppStyleSetName(), "ClassIcon.Blueprint");
diff --git a/Engine/Source/Developer/AssetTools/Private/AssetTypeActions/AssetTypeActions_Blueprint.cpp b/Engine/Source/Developer/AssetTools/Private/AssetTypeActions/AssetTypeActions_Blueprint.cpp
--- a/Engine/Source/Developer/AssetTools/Private/AssetTypeActions/AssetTypeActions_Blueprint.cpp
+++ b/Engine/Source/Developer/AssetTools/Private/AssetTypeActions/AssetTypeActions_Blueprint.cpp
@@ -224,9 +224,7 @@ FText FAssetTypeActions_Blueprint::GetNewDerivedBlueprintTooltip(TWeakObjectPtr<
 
 bool FAssetTypeActions_Blueprint::CanExecuteNewDerivedBlueprint(TWeakObjectPtr<UBlueprint> InObject)
 {
-	UBlueprint* BP = InObject.Get();
-	UClass* BPGC = BP ? BP->GeneratedClass : nullptr;
-	return BPGC && !BPGC->HasAnyClassFlags(CLASS_Deprecated);
+	return FBlueprintEditorUtils::CanCreateChildBlueprint(InObject.Get());
 }
 
 bool FAssetTypeActions_Blueprint::ShouldUseDataOnlyEditor( const UBlueprint* Blueprint ) const
diff --git a/Engine/Source/Editor/UnrealEd/Private/Kismet2/BlueprintEditorUtils.cpp b/Engine/Source/Editor/UnrealEd/Private/Kismet2/BlueprintEditorUtils.cpp
--- a/Engine/Source/Editor/UnrealEd/Private/Kismet2/BlueprintEditorUtils.cpp
+++ b/Engine/Source/Editor/UnrealEd/Private/Kismet2/BlueprintEditorUtils.cpp
@@ -3307,6 +3307,24 @@ bool FBlueprintEditorUtils::CanClassGenerateEvents(const UClass* InClass)
 	return false;
 }
 
+bool FBlueprintEditorUtils::CanCreateChildBlueprint(const UBlueprint* BP)
+{
+	if (!BP)
+	{
+		return false;
+	}
+	
+	// BP function libraries cannot have child BP's created of them. You will only ever get compilation
+	// errors if you made one.
+	if (BP->BlueprintType == EBlueprintType::BPTYPE_FunctionLibrary)
+	{
+		return false;
+	}
+
+	// Do not allow child classes to be created from deprecated BPs
+	return BP->GeneratedClass && !BP->GeneratedClass->HasAnyClassFlags(CLASS_Deprecated);
+}
+
 UEdGraph* FBlueprintEditorUtils::FindUserConstructionScript(const UBlueprint* Blueprint)
 {
 	for (UEdGraph* CurrentGraph : Blueprint->FunctionGraphs)
diff --git a/Engine/Source/Editor/UnrealEd/Public/Kismet2/BlueprintEditorUtils.h b/Engine/Source/Editor/UnrealEd/Public/Kismet2/BlueprintEditorUtils.h
--- a/Engine/Source/Editor/UnrealEd/Public/Kismet2/BlueprintEditorUtils.h
+++ b/Engine/Source/Editor/UnrealEd/Public/Kismet2/BlueprintEditorUtils.h
@@ -618,6 +618,11 @@ public:
 	 */
 	static UNREALED_API bool CanClassGenerateEvents(const UClass* Class);
 
+	/**
+	 * Returns true if this blueprint can have a child class created based on it
+	 */
+	static UNREALED_API bool CanCreateChildBlueprint(const UBlueprint* BP);
+
 	/**
 	 * If a blueprint is directly tied to a level (level script and anonymous blueprints), this will return a pointer to that level
 	 *
