do not allow for creation of a child BP function library from the content browser content menu

#rnx

diff --git a/Engine/Plugins/Editor/EngineAssetDefinitions/Source/Private/Script/AssetDefinition_Blueprint.cpp b/Engine/Plugins/Editor/EngineAssetDefinitions/Source/Private/Script/AssetDefinition_Blueprint.cpp
--- a/Engine/Plugins/Editor/EngineAssetDefinitions/Source/Private/Script/AssetDefinition_Blueprint.cpp
+++ b/Engine/Plugins/Editor/EngineAssetDefinitions/Source/Private/Script/AssetDefinition_Blueprint.cpp
@@ -222,12 +222,16 @@ namespace MenuExtension_Blueprint
 	static bool CanExecuteNewDerivedBlueprint(const FToolMenuContext& MenuContext, const FAssetData* SelectedBlueprintPtr)
 	{
 		const uint32 BPFlags = SelectedBlueprintPtr->GetTagValueRef<uint32>(FBlueprintTags::ClassFlags);
-		if ((BPFlags & (CLASS_Deprecated)) == 0)
+		if ((BPFlags & (CLASS_Deprecated)) != 0)
 		{
-			return true;
+			return false;
 		}
 
-		return false;
+		// Do not allow function libraries to have child classes created
+		const FString& BPTypeStringView = SelectedBlueprintPtr->GetTagValueRef<FString>(FBlueprintTags::BlueprintType);
+		const bool bIsFuncLib = (BPTypeStringView == TEXT("BPTYPE_FunctionLibrary")); 
+
+		return !bIsFuncLib;
 	}
 
 	static void ExecuteNewDerivedBlueprint(const FToolMenuContext& MenuContext, const FAssetData* SelectedBlueprintPtr)
@@ -312,13 +316,22 @@ namespace MenuExtension_Blueprint
 							const TAttribute<FText> ToolTip = TAttribute<FText>::CreateLambda([SelectedBlueprintPtr]()
 							{
 								const uint32 BPFlags = SelectedBlueprintPtr->GetTagValueRef<uint32>(FBlueprintTags::ClassFlags);
-								if ((BPFlags & (CLASS_Deprecated)) == 0)
+								
+								const FString BPTypeStringView = SelectedBlueprintPtr->GetTagValueRef<FString>(FBlueprintTags::BlueprintType);
+								const bool bIsFuncLib = (BPTypeStringView == TEXT("BPTYPE_FunctionLibrary")); 
+								
+								if ((BPFlags & (CLASS_Deprecated)) != 0)
 								{
-									return LOCTEXT("Blueprint_NewDerivedBlueprintTooltip", "Creates a Child Blueprint Class based on the current Blueprint, allowing you to create variants easily.");
+									return LOCTEXT("Blueprint_NewDerivedBlueprintIsDeprecatedTooltip", "Blueprint class is deprecated, cannot derive a child Blueprint!");
+								}
+								else if (bIsFuncLib)
+								{
+									return LOCTEXT("Blueprint_NewDerivedBlueprintIsFunctionLibraryTooltip", "Cannot derive from Blueprint Function Libraries!");
 								}
 								else
 								{
-									return LOCTEXT("Blueprint_NewDerivedBlueprintIsDeprecatedTooltip", "Blueprint class is deprecated, cannot derive a child Blueprint!");
+									return LOCTEXT("Blueprint_NewDerivedBlueprintTooltip", "Creates a Child Blueprint Class based on the current Blueprint, allowing you to create variants easily.");
+									
 								}
 							});
 							const FSlateIcon Icon = FSlateIcon(FAppStyle::GetAppStyleSetName(), "ClassIcon.Blueprint");
