Reorganize the Material asset creation menus

- MaterialFunction, MaterialFunctionInstance have all moved to an Advanced menu.
- NeuralProfile, SpecularProfile, and SubsurfaceProfile have all moved to a Profiles menu.
- MaterialLayer, MaterialLayerInstance, MaterialLayerBlend, and MaterialLayerBlendInstance have all moved to a Layers menu.
- By default any new assets derieved from UMaterialInterface will be in the advanced menu.
- Currently Material is left in the same location (under the basic asset creation and not in the advanced menu)

diff --git a/Engine/Plugins/Editor/EngineAssetDefinitions/Source/Private/AssetDefinition_NeuralProfile.h b/Engine/Plugins/Editor/EngineAssetDefinitions/Source/Private/AssetDefinition_NeuralProfile.h
--- a/Engine/Plugins/Editor/EngineAssetDefinitions/Source/Private/AssetDefinition_NeuralProfile.h
+++ b/Engine/Plugins/Editor/EngineAssetDefinitions/Source/Private/AssetDefinition_NeuralProfile.h
@@ -19,7 +19,7 @@ public:
 	virtual TSoftClassPtr<UObject> GetAssetClass() const override { return UNeuralProfile::StaticClass(); }
 	virtual TConstArrayView<FAssetCategoryPath> GetAssetCategories() const override
     {
-    	static const auto Categories = { EAssetCategoryPaths::Material };
+    	static const auto Categories = { EAssetCategoryPaths::Material / NSLOCTEXT("Material", "MaterialAssetSubMenu_Profiles", "Profiles") };
     	return Categories;
     }
 	// UAssetDefinition End
diff --git a/Engine/Plugins/Editor/EngineAssetDefinitions/Source/Private/AssetDefinition_SpecularProfile.h b/Engine/Plugins/Editor/EngineAssetDefinitions/Source/Private/AssetDefinition_SpecularProfile.h
--- a/Engine/Plugins/Editor/EngineAssetDefinitions/Source/Private/AssetDefinition_SpecularProfile.h
+++ b/Engine/Plugins/Editor/EngineAssetDefinitions/Source/Private/AssetDefinition_SpecularProfile.h
@@ -19,7 +19,7 @@ public:
 	virtual TSoftClassPtr<UObject> GetAssetClass() const override { return USpecularProfile::StaticClass(); }
 	virtual TConstArrayView<FAssetCategoryPath> GetAssetCategories() const override
     {
-    	static const auto Categories = { EAssetCategoryPaths::Material };
+    	static const auto Categories = { EAssetCategoryPaths::Material / NSLOCTEXT("Material", "MaterialAssetSubMenu_Profiles", "Profiles") };
     	return Categories;
     }
 	// UAssetDefinition End
diff --git a/Engine/Plugins/Editor/EngineAssetDefinitions/Source/Private/AssetDefinition_SubsurfaceProfile.h b/Engine/Plugins/Editor/EngineAssetDefinitions/Source/Private/AssetDefinition_SubsurfaceProfile.h
--- a/Engine/Plugins/Editor/EngineAssetDefinitions/Source/Private/AssetDefinition_SubsurfaceProfile.h
+++ b/Engine/Plugins/Editor/EngineAssetDefinitions/Source/Private/AssetDefinition_SubsurfaceProfile.h
@@ -19,7 +19,7 @@ public:
 	virtual TSoftClassPtr<UObject> GetAssetClass() const override { return USubsurfaceProfile::StaticClass(); }
 	virtual TConstArrayView<FAssetCategoryPath> GetAssetCategories() const override
     {
-    	static const auto Categories = { EAssetCategoryPaths::Material };
+    	static const auto Categories = { EAssetCategoryPaths::Material  / NSLOCTEXT("Material", "MaterialAssetSubMenu_Profiles", "Profiles") };
     	return Categories;
     }
 	// UAssetDefinition End
diff --git a/Engine/Plugins/Editor/EngineAssetDefinitions/Source/Private/Material/AssetDefinition_MaterialFunction.h b/Engine/Plugins/Editor/EngineAssetDefinitions/Source/Private/Material/AssetDefinition_MaterialFunction.h
--- a/Engine/Plugins/Editor/EngineAssetDefinitions/Source/Private/Material/AssetDefinition_MaterialFunction.h
+++ b/Engine/Plugins/Editor/EngineAssetDefinitions/Source/Private/Material/AssetDefinition_MaterialFunction.h
@@ -20,7 +20,7 @@ public:
 	virtual TSoftClassPtr<UObject> GetAssetClass() const override { return UMaterialFunction::StaticClass(); }
 	virtual TConstArrayView<FAssetCategoryPath> GetAssetCategories() const override
 	{
-		static const auto Categories = { EAssetCategoryPaths::Material };
+		static const auto Categories = { EAssetCategoryPaths::Material / NSLOCTEXT("Material", "MaterialAssetSubMenu_Advanced", "Advanced") };
 		return Categories;
 	}
 	virtual UThumbnailInfo* LoadThumbnailInfo(const FAssetData& InAsset) const override;
@@ -37,6 +37,11 @@ public:
 	// UAssetDefinition Begin
 	virtual FText GetAssetDisplayName() const override { return NSLOCTEXT("AssetTypeActions", "AssetTypeActions_MaterialFunctionMaterialLayer", "Material Layer"); }
 	virtual TSoftClassPtr<UObject> GetAssetClass() const override { return UMaterialFunctionMaterialLayer::StaticClass(); }
+	virtual TConstArrayView<FAssetCategoryPath> GetAssetCategories() const override
+	{
+		static const auto Categories = { EAssetCategoryPaths::Material / NSLOCTEXT("Material", "MaterialAssetSubMenu_Layers", "Layers") };
+		return Categories;
+	}
 	// UAssetDefinition End
 };
 
@@ -49,5 +54,10 @@ public:
 	// UAssetDefinition Begin
 	virtual FText GetAssetDisplayName() const override { return NSLOCTEXT("AssetTypeActions", "AssetTypeActions_MaterialFunctionMaterialLayerBlend", "Material Layer Blend"); }
 	virtual TSoftClassPtr<UObject> GetAssetClass() const override { return UMaterialFunctionMaterialLayerBlend::StaticClass(); }
+	virtual TConstArrayView<FAssetCategoryPath> GetAssetCategories() const override
+	{
+		static const auto Categories = { EAssetCategoryPaths::Material / NSLOCTEXT("Material", "MaterialAssetSubMenu_Layers", "Layers") };
+		return Categories;
+	}
 	// UAssetDefinition End
 };
diff --git a/Engine/Plugins/Editor/EngineAssetDefinitions/Source/Private/Material/AssetDefinition_MaterialFunctionInstance.h b/Engine/Plugins/Editor/EngineAssetDefinitions/Source/Private/Material/AssetDefinition_MaterialFunctionInstance.h
--- a/Engine/Plugins/Editor/EngineAssetDefinitions/Source/Private/Material/AssetDefinition_MaterialFunctionInstance.h
+++ b/Engine/Plugins/Editor/EngineAssetDefinitions/Source/Private/Material/AssetDefinition_MaterialFunctionInstance.h
@@ -32,6 +32,11 @@ public:
 	// UAssetDefinition Begin
 	virtual FText GetAssetDisplayName() const override { return NSLOCTEXT("AssetTypeActions", "AssetTypeActions_MaterialFunctionMaterialLayerInstance", "Material Layer Instance"); }
 	virtual TSoftClassPtr<UObject> GetAssetClass() const override { return UMaterialFunctionMaterialLayerInstance::StaticClass(); }
+	virtual TConstArrayView<FAssetCategoryPath> GetAssetCategories() const override
+	{
+		static const auto Categories = { EAssetCategoryPaths::Material / NSLOCTEXT("Material", "MaterialAssetSubMenu_Layers", "Layers") };
+		return Categories;
+	}
 	// UAssetDefinition End
 };
 
@@ -44,5 +49,10 @@ public:
 	// UAssetDefinition Begin
 	virtual FText GetAssetDisplayName() const override { return NSLOCTEXT("AssetTypeActions", "AssetTypeActions_MaterialFunctionMaterialLayerBlendInstance", "Material Layer Blend Instance"); }
 	virtual TSoftClassPtr<UObject> GetAssetClass() const override { return UMaterialFunctionMaterialLayerBlendInstance::StaticClass(); }
+	virtual TConstArrayView<FAssetCategoryPath> GetAssetCategories() const override
+	{
+		static const auto Categories = { EAssetCategoryPaths::Material / NSLOCTEXT("Material", "MaterialAssetSubMenu_Layers", "Layers") };
+		return Categories;
+	}
 	// UAssetDefinition End
 };
diff --git a/Engine/Plugins/Editor/EngineAssetDefinitions/Source/Private/Material/AssetDefinition_MaterialInstanceConstant.h b/Engine/Plugins/Editor/EngineAssetDefinitions/Source/Private/Material/AssetDefinition_MaterialInstanceConstant.h
--- a/Engine/Plugins/Editor/EngineAssetDefinitions/Source/Private/Material/AssetDefinition_MaterialInstanceConstant.h
+++ b/Engine/Plugins/Editor/EngineAssetDefinitions/Source/Private/Material/AssetDefinition_MaterialInstanceConstant.h
@@ -22,6 +22,11 @@ public:
 	virtual FText GetAssetDisplayName() const override { return NSLOCTEXT("AssetTypeActions", "AssetTypeActions_MaterialInstanceConstant", "Material Instance"); }
 	virtual FLinearColor GetAssetColor() const override { return FLinearColor(FColor(0,128,0)); }
 	virtual TSoftClassPtr<UObject> GetAssetClass() const override { return UMaterialInstanceConstant::StaticClass(); }
+	virtual TConstArrayView<FAssetCategoryPath> GetAssetCategories() const override
+	{
+		static const auto Categories = { EAssetCategoryPaths::Material };
+		return Categories;
+	}
 	virtual EAssetCommandResult OpenAssets(const FAssetOpenArgs& OpenArgs) const override;
 	virtual bool CanImport() const override { return true; }
 	// UAssetDefinition End
diff --git a/Engine/Plugins/Editor/EngineAssetDefinitions/Source/Private/Material/AssetDefinition_MaterialInterface.h b/Engine/Plugins/Editor/EngineAssetDefinitions/Source/Private/Material/AssetDefinition_MaterialInterface.h
--- a/Engine/Plugins/Editor/EngineAssetDefinitions/Source/Private/Material/AssetDefinition_MaterialInterface.h
+++ b/Engine/Plugins/Editor/EngineAssetDefinitions/Source/Private/Material/AssetDefinition_MaterialInterface.h
@@ -19,7 +19,7 @@ public:
 	virtual TSoftClassPtr<UObject> GetAssetClass() const override { return UMaterialInterface::StaticClass(); }
 	virtual TConstArrayView<FAssetCategoryPath> GetAssetCategories() const override
 	{
-		static const auto Categories = { EAssetCategoryPaths::Material };
+		static const auto Categories = { EAssetCategoryPaths::Material / NSLOCTEXT("Material", "MaterialAssetSubMenu_Advanced", "Advanced") };
 		return Categories;
 	}
 	virtual UThumbnailInfo* LoadThumbnailInfo(const FAssetData& InAsset) const override;
