EngineAssetDefinitions
- Adds missing export macro to publicly available DataAsset AssetDefintion

#rnx

diff --git a/Engine/Plugins/Editor/EngineAssetDefinitions/Source/Public/AssetDefinition_DataAsset.h b/Engine/Plugins/Editor/EngineAssetDefinitions/Source/Public/AssetDefinition_DataAsset.h
--- a/Engine/Plugins/Editor/EngineAssetDefinitions/Source/Public/AssetDefinition_DataAsset.h
+++ b/Engine/Plugins/Editor/EngineAssetDefinitions/Source/Public/AssetDefinition_DataAsset.h
@@ -8,7 +8,7 @@
 #include "AssetDefinition_DataAsset.generated.h"
 
 UCLASS()
-class UAssetDefinition_DataAsset : public UAssetDefinitionDefault
+class ENGINEASSETDEFINITIONS_API UAssetDefinition_DataAsset : public UAssetDefinitionDefault
 {
 	GENERATED_BODY()
 
