[ContentBrowser] Thumbnail Playback Indicator

diff --git a/Engine/Plugins/Editor/EngineAssetDefinitions/Source/EngineAssetDefinitions.Build.cs b/Engine/Plugins/Editor/EngineAssetDefinitions/Source/EngineAssetDefinitions.Build.cs
--- a/Engine/Plugins/Editor/EngineAssetDefinitions/Source/EngineAssetDefinitions.Build.cs
+++ b/Engine/Plugins/Editor/EngineAssetDefinitions/Source/EngineAssetDefinitions.Build.cs
@@ -23,6 +23,7 @@ public class EngineAssetDefinitions : ModuleRules
 		PrivateDependencyModuleNames.AddRange(
 			new string[]
 			{
+				"ApplicationCore",
 				"ContentBrowser",
 				"ContentBrowserData",
 				"AssetDefinition",
diff --git a/Engine/Plugins/Editor/EngineAssetDefinitions/Source/Private/AssetDefinition_ForceFeedbackEffect.cpp b/Engine/Plugins/Editor/EngineAssetDefinitions/Source/Private/AssetDefinition_ForceFeedbackEffect.cpp
new file mode 100644
--- /dev/null
+++ b/Engine/Plugins/Editor/EngineAssetDefinitions/Source/Private/AssetDefinition_ForceFeedbackEffect.cpp
@@ -0,0 +1,305 @@
+﻿// Copyright Epic Games, Inc. All Rights Reserved.
+
+#include "AssetDefinition_ForceFeedbackEffect.h"
+
+#include "AssetViewUtils.h"
+#include "ContentBrowserMenuContexts.h"
+#include "Framework/Application/SlateApplication.h"
+#include "ObjectEditorUtils.h"
+#include "ToolMenus.h"
+#include "Widgets/Images/SImage.h"
+#include "Widgets/Input/SButton.h"
+
+#define LOCTEXT_NAMESPACE "AssetDefinition_ForceFeedbackEffect"
+
+static FPreviewForceFeedbackEffect PreviewForceFeedbackEffect = FPreviewForceFeedbackEffect();
+
+FText UAssetDefinition_ForceFeedbackEffect::GetAssetDisplayName() const
+{
+	return LOCTEXT("AssetTypeActions_ForceFeedbackEffect", "Force Feedback Effect");
+}
+
+FLinearColor UAssetDefinition_ForceFeedbackEffect::GetAssetColor() const
+{
+	return FColor(175, 0, 0);
+}
+
+TSoftClassPtr<UObject> UAssetDefinition_ForceFeedbackEffect::GetAssetClass() const
+{
+	return UForceFeedbackEffect::StaticClass();
+}
+
+TConstArrayView<FAssetCategoryPath> UAssetDefinition_ForceFeedbackEffect::GetAssetCategories() const
+{
+	static FAssetCategoryPath Categories[] =
+		{
+		EAssetCategoryPaths::Input,
+		};
+
+	return Categories;
+}
+
+// Menu Extensions
+// --------------------------------------------------------------------
+
+namespace MenuExtension_ForceFeedbackEffect
+{
+	bool IsEffectPlaying(const TArray<TWeakObjectPtr<UForceFeedbackEffect>>& Objects)
+	{
+		if (PreviewForceFeedbackEffect.ForceFeedbackEffect)
+		{
+			for (const TWeakObjectPtr<UForceFeedbackEffect>& EffectPtr : Objects)
+			{
+				UForceFeedbackEffect* Effect = EffectPtr.Get();
+				if (Effect && PreviewForceFeedbackEffect.ForceFeedbackEffect == Effect)
+				{
+					return true;
+				}
+			}
+		}
+
+		return false;
+	}
+
+	bool IsEffectPlaying(const UForceFeedbackEffect* ForceFeedbackEffect)
+	{
+		return PreviewForceFeedbackEffect.ForceFeedbackEffect && PreviewForceFeedbackEffect.ForceFeedbackEffect == ForceFeedbackEffect;
+	}
+
+	bool IsEffectPlaying(const FAssetData& AssetData)
+	{
+		if (PreviewForceFeedbackEffect.ForceFeedbackEffect)
+		{
+			if (PreviewForceFeedbackEffect.ForceFeedbackEffect->GetFName() == AssetData.AssetName)
+			{
+				if (PreviewForceFeedbackEffect.ForceFeedbackEffect->GetOutermost()->GetFName() == AssetData.PackageName)
+				{
+					return true;
+				}
+			}
+		}
+
+		return false;
+	}
+
+	void StopEffect() 
+	{
+		PreviewForceFeedbackEffect.ResetDeviceProperties();
+		PreviewForceFeedbackEffect.ForceFeedbackEffect = nullptr;
+
+		if (IInputInterface* InputInterface = FSlateApplication::Get().GetInputInterface())
+		{
+			InputInterface->SetForceFeedbackChannelValues(0, FForceFeedbackValues());
+		}
+	}
+
+	void PlayEffect(UForceFeedbackEffect* Effect)
+	{
+		if (Effect)
+		{
+			PreviewForceFeedbackEffect.ForceFeedbackEffect = Effect;
+			PreviewForceFeedbackEffect.PlayTime = 0.f;
+			PreviewForceFeedbackEffect.PlatformUser = IPlatformInputDeviceMapper::Get().GetPrimaryPlatformUser();
+			PreviewForceFeedbackEffect.ActivateDeviceProperties();
+		}
+		else
+		{
+			StopEffect();
+		}
+	}
+
+	bool CanExecutePlayCommand(TSharedPtr<TArray<FAssetData>> InAssetsData)
+	{
+		return InAssetsData.IsValid() && InAssetsData->Num() == 1;
+	}
+
+	void ExecutePlayEffect(const TArray<TWeakObjectPtr<UForceFeedbackEffect>>& Objects)
+	{
+		for (const TWeakObjectPtr<UForceFeedbackEffect>& EffectPtr : Objects)
+		{
+			if (UForceFeedbackEffect* Effect = EffectPtr.Get())
+			{
+				// Only play the first valid effect
+				PlayEffect(Effect);
+				break;
+			}
+		}
+	}
+
+	void ExecutePlayEffect(TSharedPtr<TArray<FAssetData>> InAssetsData)
+	{
+		TArray<UObject*> SelectedForceFeedbackEffectObjects;
+		AssetViewUtils::FLoadAssetsSettings Settings{
+			// Default settings
+		};
+
+		if (InAssetsData.IsValid())
+		{
+			AssetViewUtils::LoadAssetsIfNeeded(*InAssetsData.Get(), SelectedForceFeedbackEffectObjects, Settings);
+		}
+
+		const TArray<TWeakObjectPtr<UForceFeedbackEffect>>& Objects = FObjectEditorUtils::GetTypedWeakObjectPtrs<UForceFeedbackEffect>(SelectedForceFeedbackEffectObjects);
+		ExecutePlayEffect(Objects);
+	}
+
+	void ExecuteStopEffect()
+	{
+		StopEffect();
+	}
+
+	static FDelayedAutoRegisterHelper DelayedAutoRegister(EDelayedRegisterRunPhase::EndOfEngineInit, []{ 
+		UToolMenus::RegisterStartupCallback(FSimpleMulticastDelegate::FDelegate::CreateLambda([]()
+		{
+			FToolMenuOwnerScoped OwnerScoped(UE_MODULE_NAME);
+			UToolMenu* Menu = UE::ContentBrowser::ExtendToolMenu_AssetContextMenu(UForceFeedbackEffect::StaticClass());
+
+			FToolMenuSection& Section = Menu->FindOrAddSection("GetAssetActions");
+			Section.AddDynamicEntry(NAME_None, FNewToolMenuSectionDelegate::CreateLambda([](FToolMenuSection& InSection)
+			{
+				if (const UContentBrowserAssetContextMenuContext* CBContext = UContentBrowserAssetContextMenuContext::FindContextWithAssets(InSection))
+				{
+					const TArray<FAssetData> SelectedForceFeedbackEffects = CBContext->GetSelectedAssetsOfType(UForceFeedbackEffect::StaticClass(), EIncludeSubclasses::No);
+					if (!SelectedForceFeedbackEffects.IsEmpty())
+					{
+						TSharedPtr<TArray<FAssetData>> SelectedForceFeedbackEffectsPtr = MakeShared<TArray<FAssetData>>(SelectedForceFeedbackEffects);
+						InSection.AddMenuEntry(
+							TEXT("ForceFeedbackEffect_PlayEffect"),
+							LOCTEXT("ForceFeedbackEffect_PlayEffect", "Play"),
+							LOCTEXT("ForceFeedbackEffect_PlayEffectTooltip", "Plays the selected force feedback effect."),
+							FSlateIcon(FAppStyle::GetAppStyleSetName(), "ContentBrowser.AssetAction.PlayIcon"),
+							FUIAction(
+								FExecuteAction::CreateStatic(&ExecutePlayEffect, SelectedForceFeedbackEffectsPtr),
+								FCanExecuteAction::CreateStatic(&CanExecutePlayCommand, SelectedForceFeedbackEffectsPtr)
+								)
+							);
+
+						InSection.AddMenuEntry(
+							TEXT("ForceFeedbackEffect_StopEffect"),
+							LOCTEXT("ForceFeedbackEffect_StopEffect", "Stop"),
+							LOCTEXT("ForceFeedbackEffect_StopEffectTooltip", "Stops the selected force feedback effect."),
+							FSlateIcon(FAppStyle::GetAppStyleSetName(), "ContentBrowser.AssetAction.StopIcon"),
+							FUIAction(
+								FExecuteAction::CreateStatic(&ExecuteStopEffect),
+								FCanExecuteAction()
+								)
+							);
+					}
+				}
+			}));
+		}));
+	});
+}
+
+// --------------------------------------------------------------------
+// Menu Extensions
+
+EAssetCommandResult UAssetDefinition_ForceFeedbackEffect::ActivateAssets(const FAssetActivateArgs& InActivateArgs) const
+{
+	if (InActivateArgs.ActivationMethod == EAssetActivationMethod::Previewed)
+	{
+		TArray<UObject*> Objects = InActivateArgs.LoadObjects<UObject>();
+		for (UObject* Object : Objects)
+		{
+			if (UForceFeedbackEffect* TargetEffect = Cast<UForceFeedbackEffect>(Object))
+			{
+				// Only target the first valid effect
+				TArray<TWeakObjectPtr<UForceFeedbackEffect>> EffectList;
+				EffectList.Add(MakeWeakObjectPtr(TargetEffect));
+				if (MenuExtension_ForceFeedbackEffect::IsEffectPlaying(EffectList))
+				{
+					MenuExtension_ForceFeedbackEffect::ExecuteStopEffect();
+				}
+				else
+				{
+					MenuExtension_ForceFeedbackEffect::ExecutePlayEffect(EffectList);
+				}
+
+				return EAssetCommandResult::Handled;
+			}
+		}
+	}
+	return EAssetCommandResult::Unhandled;
+}
+
+bool UAssetDefinition_ForceFeedbackEffect::GetThumbnailActionOverlay(const FAssetData& InAssetData, FAssetActionThumbnailOverlayInfo& OutActionOverlayInfo) const
+{
+	auto OnGetDisplayBrushLambda = [this, InAssetData]() -> const FSlateBrush*
+	{
+		if (MenuExtension_ForceFeedbackEffect::IsEffectPlaying(InAssetData))
+		{
+			return FAppStyle::GetBrush("ContentBrowser.AssetAction.StopIcon");
+		}
+		return FAppStyle::GetBrush("ContentBrowser.AssetAction.PlayIcon");
+	};
+
+	OutActionOverlayInfo.ActionImageWidget = SNew(SImage).Image_Lambda(OnGetDisplayBrushLambda);
+
+	auto OnClickedLambda = [InAssetData]() -> FReply
+	{
+		if (MenuExtension_ForceFeedbackEffect::IsEffectPlaying(InAssetData))
+		{
+			MenuExtension_ForceFeedbackEffect::StopEffect();
+		}
+		else
+		{
+			// Load and play asset
+			MenuExtension_ForceFeedbackEffect::PlayEffect(Cast<UForceFeedbackEffect>(InAssetData.GetAsset()));
+		}
+		return FReply::Handled();
+	};
+
+	auto OnToolTipTextLambda = [InAssetData]() -> FText
+	{
+		if (MenuExtension_ForceFeedbackEffect::IsEffectPlaying(InAssetData))
+		{
+			return LOCTEXT("Thumbnail_StopForceFeedbackToolTip", "Stop selected force feedback effect");
+		}
+		return LOCTEXT("Thumbnail_PlayForceFeedbackToolTip", "Play selected force feedback effect");
+	};
+
+	OutActionOverlayInfo.ActionButtonWidget = SNew(SButton)
+		.ButtonStyle(FAppStyle::Get(), "HoverHintOnly")
+		.ContentPadding(0.0f)
+		.ToolTipText_Lambda(OnToolTipTextLambda)
+		.OnClicked_Lambda(OnClickedLambda)
+		[
+			SNew(SImage)
+			.Image_Lambda(OnGetDisplayBrushLambda)
+		];
+
+	return true;
+}
+
+bool FPreviewForceFeedbackEffect::IsTickable() const
+{
+	return (ForceFeedbackEffect != nullptr);
+}
+
+void FPreviewForceFeedbackEffect::Tick( float DeltaTime )
+{
+	FForceFeedbackValues ForceFeedbackValues;
+
+	if (!Update(DeltaTime, ForceFeedbackValues))
+	{
+		ResetDeviceProperties();
+		ForceFeedbackEffect = nullptr;
+	}
+
+	IInputInterface* InputInterface = FSlateApplication::Get().GetInputInterface();
+	if (InputInterface)
+	{
+		InputInterface->SetForceFeedbackChannelValues(0, ForceFeedbackValues);
+	}
+}
+
+TStatId FPreviewForceFeedbackEffect::GetStatId() const
+{
+	RETURN_QUICK_DECLARE_CYCLE_STAT(FPreviewForceFeedbackEffect, STATGROUP_Tickables);
+}
+
+void FPreviewForceFeedbackEffect::AddReferencedObjects( FReferenceCollector& Collector )
+{
+	Collector.AddReferencedObject(ForceFeedbackEffect);
+}
+
+#undef LOCTEXT_NAMESPACE
diff --git a/Engine/Plugins/Editor/EngineAssetDefinitions/Source/Private/AssetDefinition_ForceFeedbackEffect.h b/Engine/Plugins/Editor/EngineAssetDefinitions/Source/Private/AssetDefinition_ForceFeedbackEffect.h
new file mode 100644
--- /dev/null
+++ b/Engine/Plugins/Editor/EngineAssetDefinitions/Source/Private/AssetDefinition_ForceFeedbackEffect.h
@@ -0,0 +1,41 @@
+// Copyright Epic Games, Inc. All Rights Reserved.
+
+#pragma once
+
+#include "AssetDefinitionDefault.h"
+#include "GameFramework/ForceFeedbackEffect.h"
+#include "GenericPlatform/IInputInterface.h"
+#include "TickableEditorObject.h"
+#include "UObject/GCObject.h"
+#include "AssetDefinition_ForceFeedbackEffect.generated.h"
+
+struct FPreviewForceFeedbackEffect : public FActiveForceFeedbackEffect, public FTickableEditorObject, public FGCObject
+{
+	// FTickableEditorObject Implementation
+	virtual bool IsTickable() const override;
+	virtual void Tick( float DeltaTime ) override;
+	virtual TStatId GetStatId() const override;
+
+	// FGCObject Implementation
+	virtual void AddReferencedObjects( FReferenceCollector& Collector ) override;
+	virtual FString GetReferencerName() const override
+	{
+		return TEXT("FPreviewForceFeedbackEffect");
+	}
+};
+
+UCLASS()
+class UAssetDefinition_ForceFeedbackEffect : public UAssetDefinitionDefault
+{
+	GENERATED_BODY()
+
+protected:
+	// UAssetDefinition Begin
+	virtual FText GetAssetDisplayName() const override;
+	virtual FLinearColor GetAssetColor() const override;
+	virtual TSoftClassPtr<UObject> GetAssetClass() const override;
+	virtual TConstArrayView<FAssetCategoryPath> GetAssetCategories() const override;
+	virtual EAssetCommandResult ActivateAssets(const FAssetActivateArgs& InActivateArgs) const override;
+	virtual bool GetThumbnailActionOverlay(const FAssetData& InAssetData, FAssetActionThumbnailOverlayInfo& OutActionOverlayInfo) const override;
+	// UAssetDefinition End
+};
diff --git a/Engine/Plugins/Experimental/AnimToTexture/Source/AnimToTextureEditor/AnimToTextureEditor.Build.cs b/Engine/Plugins/Experimental/AnimToTexture/Source/AnimToTextureEditor/AnimToTextureEditor.Build.cs
--- a/Engine/Plugins/Experimental/AnimToTexture/Source/AnimToTextureEditor/AnimToTextureEditor.Build.cs
+++ b/Engine/Plugins/Experimental/AnimToTexture/Source/AnimToTextureEditor/AnimToTextureEditor.Build.cs
@@ -25,6 +25,7 @@ public class AnimToTextureEditor : ModuleRules
 		PublicDependencyModuleNames.AddRange(
 			new string[]
 			{
+				"AssetDefinition",
 				"Core",
 				"MaterialEditor",
 				"MessageLog"
@@ -36,6 +37,8 @@ public class AnimToTextureEditor : ModuleRules
 		PrivateDependencyModuleNames.AddRange(
 			new string[]
 			{
+				"AssetTools",
+				"ContentBrowser",
 				"CoreUObject",
 				"Engine",
 				"Slate",
@@ -44,6 +47,7 @@ public class AnimToTextureEditor : ModuleRules
 				"RawMesh",
 				"MeshDescription",
 				"StaticMeshDescription",
+				"ToolMenus",
 				"UnrealEd",
 			}
 			);
diff --git a/Engine/Plugins/Experimental/AnimToTexture/Source/AnimToTextureEditor/Private/AnimToTextureAssetActions.cpp b/Engine/Plugins/Experimental/AnimToTexture/Source/AnimToTextureEditor/Private/AnimToTextureAssetActions.cpp
deleted file mode 100644
diff --git a/Engine/Plugins/Experimental/AnimToTexture/Source/AnimToTextureEditor/Private/AnimToTextureEditorModule.cpp b/Engine/Plugins/Experimental/AnimToTexture/Source/AnimToTextureEditor/Private/AnimToTextureEditorModule.cpp
--- a/Engine/Plugins/Experimental/AnimToTexture/Source/AnimToTextureEditor/Private/AnimToTextureEditorModule.cpp
+++ b/Engine/Plugins/Experimental/AnimToTexture/Source/AnimToTextureEditor/Private/AnimToTextureEditorModule.cpp
@@ -1,7 +1,6 @@
 // Copyright Epic Games, Inc. All Rights Reserved.
 
 #include "AnimToTextureEditorModule.h"
-#include "AnimToTextureAssetActions.h"
 
 #include "AssetToolsModule.h"
 #include "MessageLogModule.h"
@@ -12,11 +11,6 @@ DEFINE_LOG_CATEGORY(LogAnimToTextureEditor);
 
 void FAnimToTextureEditorModule::StartupModule()
 {
-	// Asset actions
-	FAssetToolsModule& AssetToolsModule = FModuleManager::LoadModuleChecked< FAssetToolsModule >("AssetTools");	
-	TSharedPtr<FAnimToTextureAssetActions> AnimToTextureTypeActions = MakeShareable( new FAnimToTextureAssetActions);
-	AssetToolsModule.Get().RegisterAssetTypeActions(AnimToTextureTypeActions.ToSharedRef() );
-
 	// Register Log
 	FMessageLogModule& MessageLogModule = FModuleManager::LoadModuleChecked<FMessageLogModule>("MessageLog");
 	FMessageLogInitializationOptions InitOptions;
diff --git a/Engine/Plugins/Experimental/AnimToTexture/Source/AnimToTextureEditor/Private/AssetDefinition_AnimToTexture.cpp b/Engine/Plugins/Experimental/AnimToTexture/Source/AnimToTextureEditor/Private/AssetDefinition_AnimToTexture.cpp
new file mode 100644
--- /dev/null
+++ b/Engine/Plugins/Experimental/AnimToTexture/Source/AnimToTextureEditor/Private/AssetDefinition_AnimToTexture.cpp
@@ -0,0 +1,115 @@
+// Copyright Epic Games, Inc. All Rights Reserved.
+
+#include "AssetDefinition_AnimToTexture.h"
+
+#include "AnimToTextureBPLibrary.h"
+#include "AnimToTextureDataAsset.h"
+#include "AssetViewUtils.h"
+#include "ContentBrowserMenuContexts.h"
+#include "Materials/MaterialInstanceConstant.h"
+#include "ObjectEditorUtils.h"
+#include "Styling/SlateIconFinder.h"
+#include "ToolMenus.h"
+#include "Widgets/Images/SImage.h"
+#include "Widgets/Input/SButton.h"
+
+#define LOCTEXT_NAMESPACE "UAssetDefinition_AnimToTexture"
+
+FText UAssetDefinition_AnimToTexture::GetAssetDisplayName() const
+{
+	return LOCTEXT("AnimToTextureAssetActions", "AnimToTexture");
+}
+
+FLinearColor UAssetDefinition_AnimToTexture::GetAssetColor() const
+{
+	return FColor::Blue;
+}
+
+TSoftClassPtr<UObject> UAssetDefinition_AnimToTexture::GetAssetClass() const
+{
+	return UAnimToTextureDataAsset::StaticClass();
+}
+
+TConstArrayView<FAssetCategoryPath> UAssetDefinition_AnimToTexture::GetAssetCategories() const
+{
+	static FAssetCategoryPath Categories[] =
+		{
+		EAssetCategoryPaths::Animation,
+		};
+
+	return Categories;
+}
+
+// Menu Extensions
+//--------------------------------------------------------------------
+
+namespace MenuExtension_AnimToTexture
+{
+	void RunAnimToTexture(const TArray<FAssetData> InSelectedAnimToTexture)
+	{
+		TArray<UObject*> SelectedAnimToTextureObjects;
+		AssetViewUtils::FLoadAssetsSettings Settings{
+			// Default settings
+		};
+
+		AssetViewUtils::LoadAssetsIfNeeded(InSelectedAnimToTexture, SelectedAnimToTextureObjects, Settings);
+
+		TArray<TWeakObjectPtr<UAnimToTextureDataAsset>> Objects = FObjectEditorUtils::GetTypedWeakObjectPtrs<UAnimToTextureDataAsset>(SelectedAnimToTextureObjects);
+
+		for (auto ObjIt = Objects.CreateConstIterator(); ObjIt; ++ObjIt)
+		{
+			if (UAnimToTextureDataAsset* DataAsset = ObjIt->Get())
+			{
+				// Create UVs and Textures
+				if (UAnimToTextureBPLibrary::AnimationToTexture(DataAsset))
+				{
+					// Update Material Instances (if Possible)
+					if (UStaticMesh* StaticMesh = DataAsset->GetStaticMesh())
+					{
+						for (FStaticMaterial& StaticMaterial : StaticMesh->GetStaticMaterials())
+						{
+							if (UMaterialInstanceConstant* MaterialInstanceConstant = Cast<UMaterialInstanceConstant>(StaticMaterial.MaterialInterface))
+							{
+								UAnimToTextureBPLibrary::UpdateMaterialInstanceFromDataAsset(DataAsset, MaterialInstanceConstant, EMaterialParameterAssociation::LayerParameter);
+							}
+						}
+					}
+				}
+			}
+		}
+	}
+
+	static FDelayedAutoRegisterHelper DelayedAutoRegister(EDelayedRegisterRunPhase::EndOfEngineInit, []{ 
+		UToolMenus::RegisterStartupCallback(FSimpleMulticastDelegate::FDelegate::CreateLambda([]()
+		{
+			FToolMenuOwnerScoped OwnerScoped(UE_MODULE_NAME);
+			UToolMenu* Menu = UE::ContentBrowser::ExtendToolMenu_AssetContextMenu(UAnimToTextureDataAsset::StaticClass());
+
+			FToolMenuSection& Section = Menu->FindOrAddSection("GetAssetActions");
+			Section.AddDynamicEntry(NAME_None, FNewToolMenuSectionDelegate::CreateLambda([](FToolMenuSection& InSection)
+			{
+				if (const UContentBrowserAssetContextMenuContext* CBContext = UContentBrowserAssetContextMenuContext::FindContextWithAssets(InSection))
+				{
+					const TArray<FAssetData> SelectedAnimToTexture = CBContext->GetSelectedAssetsOfType(UAnimToTextureDataAsset::StaticClass(), EIncludeSubclasses::No);
+
+					if (!SelectedAnimToTexture.IsEmpty())
+					{
+						InSection.AddMenuEntry(
+							TEXT("AnimToTexture_RunAnimationToTexture"),
+							LOCTEXT("AnimToTexture_Run", "Run Animation To Texture"),
+							LOCTEXT("AnimToTexture_RunTooltip", "Creates Vertex Animation Textures (VAT)"),
+							FSlateIcon(),
+							FUIAction(
+							FExecuteAction::CreateStatic(&RunAnimToTexture, SelectedAnimToTexture),
+							FCanExecuteAction()));
+					}
+				}
+			}));
+		}));
+	});
+};
+
+//--------------------------------------------------------------------
+// Menu Extensions
+
+#undef LOCTEXT_NAMESPACE
diff --git a/Engine/Plugins/Experimental/AnimToTexture/Source/AnimToTextureEditor/Public/AnimToTextureAssetActions.h b/Engine/Plugins/Experimental/AnimToTexture/Source/AnimToTextureEditor/Public/AnimToTextureAssetActions.h
deleted file mode 100644
diff --git a/Engine/Plugins/Experimental/AnimToTexture/Source/AnimToTextureEditor/Public/AssetDefinition_AnimToTexture.h b/Engine/Plugins/Experimental/AnimToTexture/Source/AnimToTextureEditor/Public/AssetDefinition_AnimToTexture.h
new file mode 100644
--- /dev/null
+++ b/Engine/Plugins/Experimental/AnimToTexture/Source/AnimToTextureEditor/Public/AssetDefinition_AnimToTexture.h
@@ -0,0 +1,21 @@
+// Copyright Epic Games, Inc. All Rights Reserved.
+
+#pragma once
+
+#include "AssetDefinitionDefault.h"
+
+#include "AssetDefinition_AnimToTexture.generated.h"
+
+UCLASS()
+class UAssetDefinition_AnimToTexture : public UAssetDefinitionDefault
+{
+	GENERATED_BODY()
+
+protected:
+	// UAssetDefinition Begin
+	virtual FText GetAssetDisplayName() const override;
+	virtual FLinearColor GetAssetColor() const override;
+	virtual TSoftClassPtr<UObject> GetAssetClass() const override;
+	virtual TConstArrayView<FAssetCategoryPath> GetAssetCategories() const override;
+	// UAssetDefinition End
+};
diff --git a/Engine/Plugins/Runtime/Metasound/Source/MetasoundEditor/Private/MetasoundAssetDefinitions.cpp b/Engine/Plugins/Runtime/Metasound/Source/MetasoundEditor/Private/MetasoundAssetDefinitions.cpp
--- a/Engine/Plugins/Runtime/Metasound/Source/MetasoundEditor/Private/MetasoundAssetDefinitions.cpp
+++ b/Engine/Plugins/Runtime/Metasound/Source/MetasoundEditor/Private/MetasoundAssetDefinitions.cpp
@@ -295,6 +295,72 @@ TSharedPtr<SWidget> UAssetDefinition_MetaSoundSource::GetThumbnailOverlay(const
 	return UAssetDefinition_SoundBase::GetSoundBaseThumbnailOverlay(InAssetData, MoveTemp(OnClickedLambdaOverride));
 }
 
+bool UAssetDefinition_MetaSoundSource::GetThumbnailActionOverlay(const FAssetData& InAssetData, FAssetActionThumbnailOverlayInfo& OutActionOverlayInfo) const
+{
+	auto OnGetDisplayBrushLambda = [InAssetData]() -> const FSlateBrush*
+	{
+		if (UE::AudioEditor::IsSoundPlaying(InAssetData))
+		{
+			return FAppStyle::GetBrush("ContentBrowser.AssetAction.StopIcon");
+		}
+
+		return FAppStyle::GetBrush("ContentBrowser.AssetAction.PlayIcon");
+	};
+
+	OutActionOverlayInfo.ActionImageWidget = SNew(SImage).Image_Lambda(OnGetDisplayBrushLambda);
+
+	auto OnToolTipTextLambda = [InAssetData]() -> FText
+	{
+		if (UE::AudioEditor::IsSoundPlaying(InAssetData))
+		{
+			return LOCTEXT("Thumbnail_StopSoundToolTip", "Stop selected sound");
+		}
+
+		return LOCTEXT("Thumbnail_PlaySoundToolTip", "Play selected sound");
+	};
+
+	auto OnClickedLambda = [InAssetData]() -> FReply
+	{
+		TSharedPtr<Metasound::Editor::FEditor> Editor = Metasound::Editor::FGraphBuilder::GetEditorForMetasound(*InAssetData.GetAsset());
+		if (UE::AudioEditor::IsSoundPlaying(InAssetData))
+		{
+			if (Editor.IsValid())
+			{
+				Editor->Stop();
+			}
+			else
+			{
+				UE::AudioEditor::StopSound();
+			}
+		}
+		else
+		{
+			if (Editor.IsValid())
+			{
+				Editor->Play();
+			}
+			else
+			{
+				// Load and play sound
+				UE::AudioEditor::PlaySound(Cast<USoundBase>(InAssetData.GetAsset()));
+			}
+		}
+		return FReply::Handled();
+	};
+
+	OutActionOverlayInfo.ActionButtonWidget = SNew(SButton)
+		.ButtonStyle(FAppStyle::Get(), "HoverHintOnly")
+		.ContentPadding(0.0f)
+		.ToolTipText_Lambda(OnToolTipTextLambda)
+		.OnClicked_Lambda(OnClickedLambda)
+		[
+			SNew(SImage)
+			.Image_Lambda(OnGetDisplayBrushLambda)
+		];
+
+	return true;
+}
+
 EAssetCommandResult UAssetDefinition_MetaSoundSource::ActivateAssets(const FAssetActivateArgs& ActivateArgs) const
 {
 	if (ActivateArgs.ActivationMethod == EAssetActivationMethod::Previewed)
diff --git a/Engine/Plugins/Runtime/Metasound/Source/MetasoundEditor/Private/MetasoundAssetDefinitions.h b/Engine/Plugins/Runtime/Metasound/Source/MetasoundEditor/Private/MetasoundAssetDefinitions.h
--- a/Engine/Plugins/Runtime/Metasound/Source/MetasoundEditor/Private/MetasoundAssetDefinitions.h
+++ b/Engine/Plugins/Runtime/Metasound/Source/MetasoundEditor/Private/MetasoundAssetDefinitions.h
@@ -38,6 +38,7 @@ public:
 	virtual EAssetCommandResult OpenAssets(const FAssetOpenArgs& OpenArgs) const override;
 
 	virtual TSharedPtr<SWidget> GetThumbnailOverlay(const FAssetData& InAssetData) const override;
+	virtual bool GetThumbnailActionOverlay(const FAssetData& InAssetData, FAssetActionThumbnailOverlayInfo& OutActionOverlayInfo) const override;
 	virtual EAssetCommandResult ActivateAssets(const FAssetActivateArgs& ActivateArgs) const override;
 
  	virtual const FSlateBrush* GetThumbnailBrush(const FAssetData& InAssetData, const FName InClassName) const override;
diff --git a/Engine/Source/Developer/AssetTools/Private/AssetDefinition_AssetTypeActionsProxy.cpp b/Engine/Source/Developer/AssetTools/Private/AssetDefinition_AssetTypeActionsProxy.cpp
--- a/Engine/Source/Developer/AssetTools/Private/AssetDefinition_AssetTypeActionsProxy.cpp
+++ b/Engine/Source/Developer/AssetTools/Private/AssetDefinition_AssetTypeActionsProxy.cpp
@@ -237,7 +237,14 @@ const FSlateBrush* UAssetDefinition_AssetTypeActionsProxy::GetIconBrush(const FA
 
 TSharedPtr<SWidget> UAssetDefinition_AssetTypeActionsProxy::GetThumbnailOverlay(const FAssetData& InAssetData) const
 {
+	PRAGMA_DISABLE_DEPRECATION_WARNINGS
 	return AssetType->GetThumbnailOverlay(InAssetData);
+	PRAGMA_ENABLE_DEPRECATION_WARNINGS
 }
 
-#undef LOCTEXT_NAMESPACE
\ No newline at end of file
+bool UAssetDefinition_AssetTypeActionsProxy::GetThumbnailActionOverlay(const FAssetData& InAssetData, FAssetActionThumbnailOverlayInfo& OutActionOverlayInfo) const
+{
+	return false;
+}
+
+#undef LOCTEXT_NAMESPACE
diff --git a/Engine/Source/Developer/AssetTools/Private/AssetDefinition_AssetTypeActionsProxy.h b/Engine/Source/Developer/AssetTools/Private/AssetDefinition_AssetTypeActionsProxy.h
--- a/Engine/Source/Developer/AssetTools/Private/AssetDefinition_AssetTypeActionsProxy.h
+++ b/Engine/Source/Developer/AssetTools/Private/AssetDefinition_AssetTypeActionsProxy.h
@@ -51,6 +51,7 @@ public:
 	virtual const FSlateBrush* GetThumbnailBrush(const FAssetData& InAssetData, const FName InClassName) const override;
 	virtual const FSlateBrush* GetIconBrush(const FAssetData& InAssetData, const FName InClassName) const override;
 	virtual TSharedPtr<SWidget> GetThumbnailOverlay(const FAssetData& InAssetData) const override;
+	virtual bool GetThumbnailActionOverlay(const FAssetData& InAssetData, FAssetActionThumbnailOverlayInfo& OutActionOverlayInfo) const override;
 
 protected:
 	virtual bool CanRegisterStatically() const override { return false; }
diff --git a/Engine/Source/Developer/AssetTools/Private/AssetTools.cpp b/Engine/Source/Developer/AssetTools/Private/AssetTools.cpp
--- a/Engine/Source/Developer/AssetTools/Private/AssetTools.cpp
+++ b/Engine/Source/Developer/AssetTools/Private/AssetTools.cpp
@@ -39,7 +39,6 @@
 #include "AssetTypeActions/AssetTypeActions_AnimBoneCompressionSettings.h"
 #include "AssetTypeActions/AssetTypeActions_AnimCurveCompressionSettings.h"
 #include "AssetTypeActions/AssetTypeActions_VariableFrameStrippingSettings.h"
-#include "AssetTypeActions/AssetTypeActions_ForceFeedbackEffect.h"
 #include "AssetTypeActions/AssetTypeActions_ParticleSystem.h"
 #include "AssetTypeActions/AssetTypeActions_PhysicalMaterialMask.h"
 #include "AssetTypeActions/AssetTypeActions_DeletedObjectPlaceholder.h"
@@ -396,7 +395,9 @@ public:
 	/** Optionally returns a custom widget to overlay on top of this assets' thumbnail */
     virtual TSharedPtr<class SWidget> GetThumbnailOverlay(const FAssetData& AssetData) const override
     {
+		PRAGMA_DISABLE_DEPRECATION_WARNINGS
     	return AssetDefinitionPtr.Get()->GetThumbnailOverlay(AssetData);
+		PRAGMA_ENABLE_DEPRECATION_WARNINGS
     }
 	
 	virtual EThumbnailPrimType GetDefaultThumbnailPrimitiveType(UObject* Asset) const override
@@ -1399,13 +1400,10 @@ UAssetToolsImpl::UAssetToolsImpl(const FObjectInitializer& ObjectInitializer)
 	AllocatedCategoryBits.Add(TEXT("Texture"), FAdvancedAssetCategory(EAssetTypeCategories::Textures, LOCTEXT("TextureAssetCategory", "Texture")));
 	AllocatedCategoryBits.Add(TEXT("World"), FAdvancedAssetCategory(EAssetTypeCategories::World, LOCTEXT("WorldAssetCategory", "World")));
 
-	EAssetTypeCategories::Type InputCategoryBit = RegisterAdvancedAssetCategory(FName(TEXT("Input")), LOCTEXT("InputAssetsCategory", "Input"));
-	
 	// Register the built-in asset type actions
 	RegisterAssetTypeActions(MakeShareable(new FAssetTypeActions_AnimBoneCompressionSettings));
 	RegisterAssetTypeActions(MakeShareable(new FAssetTypeActions_AnimCurveCompressionSettings));
 	RegisterAssetTypeActions(MakeShareable(new FAssetTypeActions_VariableFrameStrippingSettings));
-	RegisterAssetTypeActions(MakeShareable(new FAssetTypeActions_ForceFeedbackEffect(InputCategoryBit)));
 	RegisterAssetTypeActions(MakeShareable(new FAssetTypeActions_ParticleSystem));
 	RegisterAssetTypeActions(MakeShareable(new FAssetTypeActions_PhysicalMaterialMask));
 	RegisterAssetTypeActions(MakeShareable(new FAssetTypeActions_DeletedObjectPlaceholder));
diff --git a/Engine/Source/Developer/AssetTools/Private/AssetTypeActions/AssetTypeActions_ForceFeedbackEffect.cpp b/Engine/Source/Developer/AssetTools/Private/AssetTypeActions/AssetTypeActions_ForceFeedbackEffect.cpp
deleted file mode 100644
diff --git a/Engine/Source/Developer/AssetTools/Private/AssetTypeActions/AssetTypeActions_ForceFeedbackEffect.h b/Engine/Source/Developer/AssetTools/Private/AssetTypeActions/AssetTypeActions_ForceFeedbackEffect.h
deleted file mode 100644
diff --git a/Engine/Source/Developer/AssetTools/Public/IAssetTypeActions.h b/Engine/Source/Developer/AssetTools/Public/IAssetTypeActions.h
--- a/Engine/Source/Developer/AssetTools/Public/IAssetTypeActions.h
+++ b/Engine/Source/Developer/AssetTools/Public/IAssetTypeActions.h
@@ -116,6 +116,7 @@ public:
 	virtual EThumbnailPrimType GetDefaultThumbnailPrimitiveType(UObject* Asset) const = 0;
 
 	/** Optionally returns a custom widget to overlay on top of this assets' thumbnail */
+	UE_DEPRECATED(5.5, "Convert the Action class over to the Definition system and use the new GetThumbnailActionOverlay, see UAssetDefinition for more information")
 	virtual TSharedPtr<class SWidget> GetThumbnailOverlay(const FAssetData& AssetData) const = 0;
 
 	/** Returns additional tooltip information for the specified asset, if it has any (otherwise return the null widget) */
diff --git a/Engine/Source/Editor/AssetDefinition/Public/AssetDefinition.h b/Engine/Source/Editor/AssetDefinition/Public/AssetDefinition.h
--- a/Engine/Source/Editor/AssetDefinition/Public/AssetDefinition.h
+++ b/Engine/Source/Editor/AssetDefinition/Public/AssetDefinition.h
@@ -325,6 +325,16 @@ public:
 };
 
 
+struct FAssetActionThumbnailOverlayInfo
+{
+public:
+	/** Image widget displayed in the top left border */
+	TSharedPtr<SWidget> ActionImageWidget;
+
+	/** Button widget displayed in the center when hovering the thumbnail */
+	TSharedPtr<SWidget> ActionButtonWidget;
+};
+
 /**
  * Asset Definitions represent top level assets that are known to the editor.
  *
@@ -565,13 +575,19 @@ public:
 	{
 		return nullptr;
 	}
-	
+
+	UE_DEPRECATED(5.5, "Please use the new GetThumbnailActionOverlay function")
 	/** Optionally returns a custom widget to overlay on top of this assets' thumbnail */
 	virtual TSharedPtr<SWidget> GetThumbnailOverlay(const FAssetData& InAssetData) const
 	{
 		return TSharedPtr<SWidget>();
 	}
 
+	virtual bool GetThumbnailActionOverlay(const FAssetData& InAssetData, FAssetActionThumbnailOverlayInfo& OutActionOverlayInfo) const
+	{
+		return false;
+	}
+
 	/** Whether this asset has external packages associated with it (impacts saving and dirty-state behavior) */
 	virtual bool ShouldSaveExternalPackages() const
 	{
diff --git a/Engine/Source/Editor/AudioEditor/Private/AssetTypeActions/AssetDefinition_SoundBase.cpp b/Engine/Source/Editor/AudioEditor/Private/AssetTypeActions/AssetDefinition_SoundBase.cpp
--- a/Engine/Source/Editor/AudioEditor/Private/AssetTypeActions/AssetDefinition_SoundBase.cpp
+++ b/Engine/Source/Editor/AudioEditor/Private/AssetTypeActions/AssetDefinition_SoundBase.cpp
@@ -195,6 +195,57 @@ TSharedPtr<SWidget> UAssetDefinition_SoundBase::GetThumbnailOverlay(const FAsset
 	return GetSoundBaseThumbnailOverlay(InAssetData, MoveTemp(OnClickedLambda));
 }
 
+bool UAssetDefinition_SoundBase::GetThumbnailActionOverlay(const FAssetData& InAssetData, FAssetActionThumbnailOverlayInfo& OutActionOverlayInfo) const
+{
+	auto OnGetDisplayBrushLambda = [InAssetData]() -> const FSlateBrush*
+	{
+		if (UE::AudioEditor::IsSoundPlaying(InAssetData))
+		{
+			return FAppStyle::GetBrush("ContentBrowser.AssetAction.StopIcon");
+		}
+
+		return FAppStyle::GetBrush("ContentBrowser.AssetAction.PlayIcon");
+	};
+
+	OutActionOverlayInfo.ActionImageWidget = SNew(SImage).Image_Lambda(OnGetDisplayBrushLambda);
+
+	auto OnToolTipTextLambda = [InAssetData]() -> FText
+	{
+		if (UE::AudioEditor::IsSoundPlaying(InAssetData))
+		{
+			return LOCTEXT("Thumbnail_StopSoundToolTip", "Stop selected sound");
+		}
+
+		return LOCTEXT("Thumbnail_PlaySoundToolTip", "Play selected sound");
+	};
+
+	auto OnClickedLambda = [InAssetData]() -> FReply
+	{
+		if (UE::AudioEditor::IsSoundPlaying(InAssetData))
+		{
+			UE::AudioEditor::StopSound();
+		}
+		else
+		{
+			// Load and play sound
+			UE::AudioEditor::PlaySound(Cast<USoundBase>(InAssetData.GetAsset()));
+		}
+		return FReply::Handled();
+	};
+
+	OutActionOverlayInfo.ActionButtonWidget = SNew(SButton)
+		.ButtonStyle(FAppStyle::Get(), "HoverHintOnly")
+		.ContentPadding(0.0f)
+		.ToolTipText_Lambda(OnToolTipTextLambda)
+		.OnClicked_Lambda(OnClickedLambda)
+		[
+			SNew(SImage)
+			.Image_Lambda(OnGetDisplayBrushLambda)
+		];
+
+	return true;
+}
+
 // Menu Extensions
 //--------------------------------------------------------------------
 
diff --git a/Engine/Source/Editor/AudioEditor/Private/AssetTypeActions/AssetDefinition_SoundSourceBus.h b/Engine/Source/Editor/AudioEditor/Private/AssetTypeActions/AssetDefinition_SoundSourceBus.h
--- a/Engine/Source/Editor/AudioEditor/Private/AssetTypeActions/AssetDefinition_SoundSourceBus.h
+++ b/Engine/Source/Editor/AudioEditor/Private/AssetTypeActions/AssetDefinition_SoundSourceBus.h
@@ -24,5 +24,6 @@ public:
 		return Categories;
 	}
 	virtual TSharedPtr<SWidget> GetThumbnailOverlay(const FAssetData& InAssetData) const override { return nullptr; }
+	virtual bool GetThumbnailActionOverlay(const FAssetData& InAssetData, FAssetActionThumbnailOverlayInfo& OutActionOverlayInfo) const override { return false; }
 	// UAssetDefinition End
 };
diff --git a/Engine/Source/Editor/AudioEditor/Private/AudioEditorModule.cpp b/Engine/Source/Editor/AudioEditor/Private/AudioEditorModule.cpp
--- a/Engine/Source/Editor/AudioEditor/Private/AudioEditorModule.cpp
+++ b/Engine/Source/Editor/AudioEditor/Private/AudioEditorModule.cpp
@@ -179,7 +179,6 @@ public:
 
 		AssetTools.RegisterAssetTypeActions(MakeShared<FAssetTypeActions_DialogueWave>());
 		AssetTools.RegisterAssetTypeActions(MakeShared<FAssetTypeActions_SoundClass>());
-		AssetTools.RegisterAssetTypeActions(MakeShared<FAssetTypeActions_SoundWave>());
 	}
 
 	virtual void RegisterAudioMixerAssetActions() override
diff --git a/Engine/Source/Editor/AudioEditor/Public/AssetTypeActions/AssetDefinition_SoundBase.h b/Engine/Source/Editor/AudioEditor/Public/AssetTypeActions/AssetDefinition_SoundBase.h
--- a/Engine/Source/Editor/AudioEditor/Public/AssetTypeActions/AssetDefinition_SoundBase.h
+++ b/Engine/Source/Editor/AudioEditor/Public/AssetTypeActions/AssetDefinition_SoundBase.h
@@ -35,6 +35,7 @@ public:
 		return Categories;
 	}
 	virtual TSharedPtr<SWidget> GetThumbnailOverlay(const FAssetData& InAssetData) const override;
+	virtual bool GetThumbnailActionOverlay(const FAssetData& InAssetData, FAssetActionThumbnailOverlayInfo& OutActionOverlayInfo) const override;
 	virtual EAssetCommandResult ActivateAssets(const FAssetActivateArgs& ActivateArgs) const override;
 	// UAssetDefinition End
 
diff --git a/Engine/Source/Editor/EditorStyle/Private/StarshipStyle.cpp b/Engine/Source/Editor/EditorStyle/Private/StarshipStyle.cpp
--- a/Engine/Source/Editor/EditorStyle/Private/StarshipStyle.cpp
+++ b/Engine/Source/Editor/EditorStyle/Private/StarshipStyle.cpp
@@ -6462,7 +6462,12 @@ void FStarshipEditorStyle::FStyle::SetupContentBrowserStyle()
 		// Misc
 		/** Should be moved, shared */ Set( "ContentBrowser.ThumbnailShadow", new BOX_BRUSH( "ContentBrowser/ThumbnailShadow" , FMargin( 4.0f / 64.0f ) ) );
 
-
+		// Playback Action
+		Set( "ContentBrowser.AssetAction.PlayIcon", new IMAGE_BRUSH_SVG("Starship/ContentBrowser/PlayIcon", Icon16x16));
+		Set( "ContentBrowser.AssetAction.StopIcon", new IMAGE_BRUSH_SVG("Starship/ContentBrowser/StopIcon", Icon16x16));
+		FLinearColor OverlayColorAssetThumbnailOverlay = FStyleColors::Panel.GetSpecifiedColor();
+		OverlayColorAssetThumbnailOverlay.A = 0.75f;
+		Set("ContentBrowser.AssetTileItem.AssetThumbnailBar", new FSlateRoundedBoxBrush(OverlayColorAssetThumbnailOverlay, 4.f));
 
 		Set( "ContentBrowser.ColumnViewAssetIcon", new IMAGE_BRUSH( "Icons/doc_16x", Icon16x16 ) );
 
diff --git a/Engine/Source/Editor/UnrealEd/Private/AssetThumbnail.cpp b/Engine/Source/Editor/UnrealEd/Private/AssetThumbnail.cpp
--- a/Engine/Source/Editor/UnrealEd/Private/AssetThumbnail.cpp
+++ b/Engine/Source/Editor/UnrealEd/Private/AssetThumbnail.cpp
@@ -33,6 +33,7 @@
 #include "AssetRegistry/IAssetRegistry.h"
 #include "ShaderCompiler.h"
 #include "AssetCompilingManager.h"
+#include "AssetDefinitionRegistry.h"
 #include "IAssetTools.h"
 #include "AssetTypeActions_Base.h"
 #include "AssetToolsModule.h"
@@ -330,14 +331,75 @@ public:
 
 		if( InArgs._AllowAssetSpecificThumbnailOverlay && AssetTypeActions.IsValid() )
 		{
-			// Does the asset provide an additional thumbnail overlay?
-			TSharedPtr<SWidget> AssetSpecificThumbnailOverlay = AssetTypeActions->GetThumbnailOverlay(AssetData);
-			if( AssetSpecificThumbnailOverlay.IsValid() )
+#if UE_CONTENTBROWSER_NEW_STYLE
+			const UAssetDefinition* AssetDefinition = UAssetDefinitionRegistry::Get()->GetAssetDefinitionForClass(Class);
+
+			FAssetActionThumbnailOverlayInfo OutThumbnailInfo;
+			if (AssetDefinition && AssetDefinition->GetThumbnailActionOverlay(AssetData, OutThumbnailInfo))
 			{
-				OverlayWidget->AddSlot()
-				[
-					AssetSpecificThumbnailOverlay.ToSharedRef()
-				];
+				constexpr int32 OverlayZOrder = 1;
+				if (OutThumbnailInfo.ActionImageWidget.IsValid())
+				{
+					constexpr float PaddingFromTopLeftBorder = 2.f;
+					constexpr float TopLeftImageSize = 20.f;
+
+					OverlayWidget->AddSlot()
+					.ZOrder(OverlayZOrder)
+					.VAlign(VAlign_Top)
+					.HAlign(HAlign_Left)
+					.Padding(PaddingFromTopLeftBorder, PaddingFromTopLeftBorder, 0.f, 0.f)
+					[
+						SNew(SBorder)
+						.BorderImage(FAppStyle::GetBrush(TEXT("ContentBrowser.AssetTileItem.AssetThumbnailBar")))
+						.Visibility_Lambda([this] () { return IsHovered() ? EVisibility::Collapsed : EVisibility::Visible; })
+						[
+							SNew(SBox)
+							.WidthOverride(TopLeftImageSize)
+							.HeightOverride(TopLeftImageSize)
+							[
+								OutThumbnailInfo.ActionImageWidget.ToSharedRef()
+							]
+						]
+					];
+				}
+
+				if (OutThumbnailInfo.ActionButtonWidget.IsValid())
+				{
+					constexpr float CenterImageSize = 32.f;
+
+					OverlayWidget->AddSlot()
+					.ZOrder(OverlayZOrder)
+					.VAlign(VAlign_Center)
+					.HAlign(HAlign_Center)
+					[
+						SNew(SBorder)
+						.BorderImage(FAppStyle::GetBrush(TEXT("ContentBrowser.AssetTileItem.AssetThumbnailBar")))
+						.Visibility_Lambda([this] () { return IsHovered() ? EVisibility::Visible : EVisibility::Collapsed; })
+						[
+							SNew(SBox)
+							.WidthOverride(CenterImageSize)
+							.HeightOverride(CenterImageSize)
+							[
+								OutThumbnailInfo.ActionButtonWidget.ToSharedRef()
+							]
+						]
+					];
+				}
+			}
+			else
+#endif
+			{
+				// Does the asset provide an additional thumbnail overlay?
+				PRAGMA_DISABLE_DEPRECATION_WARNINGS
+				TSharedPtr<SWidget> AssetSpecificThumbnailOverlay = AssetTypeActions->GetThumbnailOverlay(AssetData);
+				PRAGMA_ENABLE_DEPRECATION_WARNINGS
+				if( AssetSpecificThumbnailOverlay.IsValid() )
+				{
+					OverlayWidget->AddSlot()
+					[
+						AssetSpecificThumbnailOverlay.ToSharedRef()
+					];
+				}
 			}
 		}
 
