Fix a startup crash in monolithic editor targets due to a static variable that attempts to construct a UObject
 - Move static variable into a static accessor that it is created on first access rather than during global initialization

diff --git a/Engine/Plugins/Editor/EngineAssetDefinitions/Source/Private/AssetDefinition_ForceFeedbackEffect.cpp b/Engine/Plugins/Editor/EngineAssetDefinitions/Source/Private/AssetDefinition_ForceFeedbackEffect.cpp
--- a/Engine/Plugins/Editor/EngineAssetDefinitions/Source/Private/AssetDefinition_ForceFeedbackEffect.cpp
+++ b/Engine/Plugins/Editor/EngineAssetDefinitions/Source/Private/AssetDefinition_ForceFeedbackEffect.cpp
@@ -1,4 +1,4 @@
-﻿// Copyright Epic Games, Inc. All Rights Reserved.
+// Copyright Epic Games, Inc. All Rights Reserved.
 
 #include "AssetDefinition_ForceFeedbackEffect.h"
 
@@ -12,8 +12,6 @@
 
 #define LOCTEXT_NAMESPACE "AssetDefinition_ForceFeedbackEffect"
 
-FPreviewForceFeedbackEffect UAssetDefinition_ForceFeedbackEffect::PreviewForceFeedbackEffect = FPreviewForceFeedbackEffect();
-
 FText UAssetDefinition_ForceFeedbackEffect::GetAssetDisplayName() const
 {
 	return LOCTEXT("AssetTypeActions_ForceFeedbackEffect", "Force Feedback Effect");
@@ -42,12 +40,12 @@ namespace MenuExtension_ForceFeedbackEffect
 {
 	bool IsEffectPlaying(const TArray<TWeakObjectPtr<UForceFeedbackEffect>>& Objects)
 	{
-		if (UAssetDefinition_ForceFeedbackEffect::PreviewForceFeedbackEffect.ForceFeedbackEffect)
+		if (UAssetDefinition_ForceFeedbackEffect::GetPreviewForceFeedbackEffect().ForceFeedbackEffect)
 		{
 			for (const TWeakObjectPtr<UForceFeedbackEffect>& EffectPtr : Objects)
 			{
 				UForceFeedbackEffect* Effect = EffectPtr.Get();
-				if (Effect && UAssetDefinition_ForceFeedbackEffect::PreviewForceFeedbackEffect.ForceFeedbackEffect == Effect)
+				if (Effect && UAssetDefinition_ForceFeedbackEffect::GetPreviewForceFeedbackEffect().ForceFeedbackEffect == Effect)
 				{
 					return true;
 				}
@@ -59,16 +57,16 @@ namespace MenuExtension_ForceFeedbackEffect
 
 	bool IsEffectPlaying(const UForceFeedbackEffect* ForceFeedbackEffect)
 	{
-		return UAssetDefinition_ForceFeedbackEffect::PreviewForceFeedbackEffect.ForceFeedbackEffect && UAssetDefinition_ForceFeedbackEffect::PreviewForceFeedbackEffect.ForceFeedbackEffect == ForceFeedbackEffect;
+		return UAssetDefinition_ForceFeedbackEffect::GetPreviewForceFeedbackEffect().ForceFeedbackEffect && UAssetDefinition_ForceFeedbackEffect::GetPreviewForceFeedbackEffect().ForceFeedbackEffect == ForceFeedbackEffect;
 	}
 
 	bool IsEffectPlaying(const FAssetData& AssetData)
 	{
-		if (UAssetDefinition_ForceFeedbackEffect::PreviewForceFeedbackEffect.ForceFeedbackEffect)
+		if (UAssetDefinition_ForceFeedbackEffect::GetPreviewForceFeedbackEffect().ForceFeedbackEffect)
 		{
-			if (UAssetDefinition_ForceFeedbackEffect::PreviewForceFeedbackEffect.ForceFeedbackEffect->GetFName() == AssetData.AssetName)
+			if (UAssetDefinition_ForceFeedbackEffect::GetPreviewForceFeedbackEffect().ForceFeedbackEffect->GetFName() == AssetData.AssetName)
 			{
-				if (UAssetDefinition_ForceFeedbackEffect::PreviewForceFeedbackEffect.ForceFeedbackEffect->GetOutermost()->GetFName() == AssetData.PackageName)
+				if (UAssetDefinition_ForceFeedbackEffect::GetPreviewForceFeedbackEffect().ForceFeedbackEffect->GetOutermost()->GetFName() == AssetData.PackageName)
 				{
 					return true;
 				}
@@ -80,8 +78,8 @@ namespace MenuExtension_ForceFeedbackEffect
 
 	void StopEffect() 
 	{
-		UAssetDefinition_ForceFeedbackEffect::PreviewForceFeedbackEffect.ResetDeviceProperties();
-		UAssetDefinition_ForceFeedbackEffect::PreviewForceFeedbackEffect.ForceFeedbackEffect = nullptr;
+		UAssetDefinition_ForceFeedbackEffect::GetPreviewForceFeedbackEffect().ResetDeviceProperties();
+		UAssetDefinition_ForceFeedbackEffect::GetPreviewForceFeedbackEffect().ForceFeedbackEffect = nullptr;
 
 		if (IInputInterface* InputInterface = FSlateApplication::Get().GetInputInterface())
 		{
@@ -93,10 +91,10 @@ namespace MenuExtension_ForceFeedbackEffect
 	{
 		if (Effect)
 		{
-			UAssetDefinition_ForceFeedbackEffect::PreviewForceFeedbackEffect.ForceFeedbackEffect = Effect;
-			UAssetDefinition_ForceFeedbackEffect::PreviewForceFeedbackEffect.PlayTime = 0.f;
-			UAssetDefinition_ForceFeedbackEffect::PreviewForceFeedbackEffect.PlatformUser = IPlatformInputDeviceMapper::Get().GetPrimaryPlatformUser();
-			UAssetDefinition_ForceFeedbackEffect::PreviewForceFeedbackEffect.ActivateDeviceProperties();
+			UAssetDefinition_ForceFeedbackEffect::GetPreviewForceFeedbackEffect().ForceFeedbackEffect = Effect;
+			UAssetDefinition_ForceFeedbackEffect::GetPreviewForceFeedbackEffect().PlayTime = 0.f;
+			UAssetDefinition_ForceFeedbackEffect::GetPreviewForceFeedbackEffect().PlatformUser = IPlatformInputDeviceMapper::Get().GetPrimaryPlatformUser();
+			UAssetDefinition_ForceFeedbackEffect::GetPreviewForceFeedbackEffect().ActivateDeviceProperties();
 		}
 		else
 		{
diff --git a/Engine/Plugins/Editor/EngineAssetDefinitions/Source/Private/AssetDefinition_ForceFeedbackEffect.h b/Engine/Plugins/Editor/EngineAssetDefinitions/Source/Private/AssetDefinition_ForceFeedbackEffect.h
--- a/Engine/Plugins/Editor/EngineAssetDefinitions/Source/Private/AssetDefinition_ForceFeedbackEffect.h
+++ b/Engine/Plugins/Editor/EngineAssetDefinitions/Source/Private/AssetDefinition_ForceFeedbackEffect.h
@@ -40,5 +40,10 @@ protected:
 	// UAssetDefinition End
 
 public:
-	static FPreviewForceFeedbackEffect PreviewForceFeedbackEffect;
+
+	static FPreviewForceFeedbackEffect& GetPreviewForceFeedbackEffect()
+	{
+		static FPreviewForceFeedbackEffect Instance;
+		return Instance;
+	}
 };
