[EngineAssetDefinition] Regression: ForceFeedbackEffect Asset do not have the Play Icon in the old ContentBrowser

diff --git a/Engine/Plugins/Editor/EngineAssetDefinitions/Source/Private/AssetDefinition_ForceFeedbackEffect.cpp b/Engine/Plugins/Editor/EngineAssetDefinitions/Source/Private/AssetDefinition_ForceFeedbackEffect.cpp
--- a/Engine/Plugins/Editor/EngineAssetDefinitions/Source/Private/AssetDefinition_ForceFeedbackEffect.cpp
+++ b/Engine/Plugins/Editor/EngineAssetDefinitions/Source/Private/AssetDefinition_ForceFeedbackEffect.cpp
@@ -211,6 +211,78 @@ EAssetCommandResult UAssetDefinition_ForceFeedbackEffect::ActivateAssets(const F
 	return EAssetCommandResult::Unhandled;
 }
 
+TSharedPtr<SWidget> UAssetDefinition_ForceFeedbackEffect::GetThumbnailOverlay(const FAssetData& InAssetData) const
+{
+	auto OnGetDisplayBrushLambda = [this, InAssetData]() -> const FSlateBrush*
+	{
+		if (MenuExtension_ForceFeedbackEffect::IsEffectPlaying(InAssetData))
+		{
+			return FAppStyle::GetBrush("MediaAsset.AssetActions.Stop.Large");
+		}
+		return FAppStyle::GetBrush("MediaAsset.AssetActions.Play.Large");
+	};
+
+	auto OnClickedLambda = [InAssetData]() -> FReply
+	{
+		if (MenuExtension_ForceFeedbackEffect::IsEffectPlaying(InAssetData))
+		{
+			MenuExtension_ForceFeedbackEffect::StopEffect();
+		}
+		else
+		{
+			// Load and play asset
+			MenuExtension_ForceFeedbackEffect::PlayEffect(Cast<UForceFeedbackEffect>(InAssetData.GetAsset()));
+		}
+		return FReply::Handled();
+	};
+
+	auto OnToolTipTextLambda = [this, InAssetData]() -> FText
+	{
+		if (MenuExtension_ForceFeedbackEffect::IsEffectPlaying(InAssetData))
+		{
+			return LOCTEXT("Thumbnail_StopForceFeedbackToolTip", "Stop selected force feedback effect");
+		}
+		return LOCTEXT("Thumbnail_PlayForceFeedbackToolTip", "Play selected force feedback effect");
+	};
+
+	TSharedRef<SBox> Box = SNew(SBox)
+		.HAlign(HAlign_Center)
+		.VAlign(VAlign_Center)
+		.Padding(FMargin(2));
+
+	auto OnGetVisibilityLambda = [this, Box, InAssetData]() -> EVisibility
+	{
+		if (Box->IsHovered() || MenuExtension_ForceFeedbackEffect::IsEffectPlaying(InAssetData))
+		{
+			return EVisibility::Visible;
+		}
+		return EVisibility::Hidden;
+	};
+
+	TSharedRef<SButton> BoxContent = SNew(SButton)
+		.ButtonStyle(FAppStyle::Get(), "HoverHintOnly")
+		.ToolTipText_Lambda(OnToolTipTextLambda)
+		.Cursor(EMouseCursor::Default) // The outer widget can specify a DragHand cursor, so we need to override that here
+		.ForegroundColor(FSlateColor::UseForeground())
+		.IsFocusable(false)
+		.OnClicked_Lambda(OnClickedLambda)
+		.Visibility_Lambda(OnGetVisibilityLambda)
+		[
+			SNew(SBox)
+			.MinDesiredWidth(16.f)
+			.MinDesiredHeight(16.f)
+			[
+				SNew(SImage)
+				.Image_Lambda(OnGetDisplayBrushLambda)
+			]
+		];
+
+	Box->SetContent(BoxContent);
+	Box->SetVisibility(EVisibility::Visible);
+
+	return Box;
+}
+
 bool UAssetDefinition_ForceFeedbackEffect::GetThumbnailActionOverlay(const FAssetData& InAssetData, FAssetActionThumbnailOverlayInfo& OutActionOverlayInfo) const
 {
 	auto OnGetDisplayBrushLambda = [this, InAssetData]() -> const FSlateBrush*
diff --git a/Engine/Plugins/Editor/EngineAssetDefinitions/Source/Private/AssetDefinition_ForceFeedbackEffect.h b/Engine/Plugins/Editor/EngineAssetDefinitions/Source/Private/AssetDefinition_ForceFeedbackEffect.h
--- a/Engine/Plugins/Editor/EngineAssetDefinitions/Source/Private/AssetDefinition_ForceFeedbackEffect.h
+++ b/Engine/Plugins/Editor/EngineAssetDefinitions/Source/Private/AssetDefinition_ForceFeedbackEffect.h
@@ -36,6 +36,7 @@ protected:
 	virtual TSoftClassPtr<UObject> GetAssetClass() const override;
 	virtual TConstArrayView<FAssetCategoryPath> GetAssetCategories() const override;
 	virtual EAssetCommandResult ActivateAssets(const FAssetActivateArgs& InActivateArgs) const override;
+	virtual TSharedPtr<SWidget> GetThumbnailOverlay(const FAssetData& InAssetData) const override;
 	virtual bool GetThumbnailActionOverlay(const FAssetData& InAssetData, FAssetActionThumbnailOverlayInfo& OutActionOverlayInfo) const override;
 	// UAssetDefinition End
 
