Improve text diff for maps

[RN] Diffing umap assets now provides a more useful text representation of the map, rather than an empty T3D object

diff --git a/Engine/Plugins/Editor/EngineAssetDefinitions/Source/EngineAssetDefinitions.Build.cs b/Engine/Plugins/Editor/EngineAssetDefinitions/Source/EngineAssetDefinitions.Build.cs
--- a/Engine/Plugins/Editor/EngineAssetDefinitions/Source/EngineAssetDefinitions.Build.cs
+++ b/Engine/Plugins/Editor/EngineAssetDefinitions/Source/EngineAssetDefinitions.Build.cs
@@ -26,6 +26,7 @@ public class EngineAssetDefinitions : ModuleRules
 				"ApplicationCore",
 				"ContentBrowser",
 				"ContentBrowserData",
+				"JsonObjectGraph",
 				"AssetDefinition",
 				"UnrealEd",
 				"ToolMenus",
diff --git a/Engine/Plugins/Editor/EngineAssetDefinitions/Source/Private/AssetDefinition_World.cpp b/Engine/Plugins/Editor/EngineAssetDefinitions/Source/Private/AssetDefinition_World.cpp
--- a/Engine/Plugins/Editor/EngineAssetDefinitions/Source/Private/AssetDefinition_World.cpp
+++ b/Engine/Plugins/Editor/EngineAssetDefinitions/Source/Private/AssetDefinition_World.cpp
@@ -2,8 +2,12 @@
 
 #include "AssetDefinition_World.h"
 
+#include "AssetToolsModule.h"
 #include "FileHelpers.h"
+#include "HAL/FileManager.h"
+#include "JsonObjectGraph/Stringify.h"
 #include "Misc/PackageName.h"
+#include "Settings/EditorLoadingSavingSettings.h"
 #include "ThumbnailRendering/WorldThumbnailInfo.h"
 #include "Misc/MessageDialog.h"
 #include "Engine/Level.h"
@@ -69,6 +73,56 @@ EAssetCommandResult UAssetDefinition_World::OpenAssets(const FAssetOpenArgs& Ope
 	return EAssetCommandResult::Unhandled;
 }
 
+EAssetCommandResult UAssetDefinition_World::PerformAssetDiff(const FAssetDiffArgs& DiffArgs) const
+{
+	if (DiffArgs.OldAsset == nullptr && DiffArgs.NewAsset == nullptr)
+	{
+		return EAssetCommandResult::Unhandled;
+	}
+
+	// The caller has loaded our assets and classified them, create a useful text representation for display:
+	const FUtf8String OldAssetJSON = DiffArgs.OldAsset ? UE::JsonObjectGraph::Stringify({DiffArgs.OldAsset->GetPackage()}) : FUtf8String();
+	const FUtf8String NewAssetJSON = DiffArgs.NewAsset ? UE::JsonObjectGraph::Stringify({DiffArgs.NewAsset->GetPackage()}) : FUtf8String();
+
+	// Write the JSON to a file, so the TextDiffTool can consume it:
+	const auto GetFilename = [](UObject* Asset, const FString& Revision)
+	{
+		static int32 ID = 0; // ensure subsequent diffs within a session don't stomp eachother, matching UAssetToolsImpl::DumpAssetToTempFile
+		FString AssetName = Asset? Asset->GetName() : TEXT("empty");
+		FString RelTempFileName = FString::Printf(TEXT("%sJsonDiff%s-%s-%d.txt"), *FPaths::DiffDir(), *AssetName, *Revision, ID++);
+		FString AbsoluteTempFileName = FPaths::ConvertRelativePathToFull(RelTempFileName);
+		return AbsoluteTempFileName ;
+	};
+
+	const FString OldFilename = GetFilename(DiffArgs.OldAsset, DiffArgs.OldRevision.Revision);
+	const FString NewFilename = GetFilename(DiffArgs.NewAsset, DiffArgs.NewRevision.Revision);
+
+	const auto WriteUtf8StringToFile = [](const FString& Filename, const FUtf8String& String)
+	{
+		TUniquePtr<FArchive> FileArchive(IFileManager::Get().CreateDebugFileWriter(*Filename));
+		if(FileArchive)
+		{
+			FileArchive->Serialize((void*)String.GetCharArray().GetData(), FMath::Max(String.GetCharArray().Num() - 1, 0));
+			// FileArchive closed by ~FArchiveFileWriterGeneric
+			return true;
+		}
+		return false;
+	};
+	if( !WriteUtf8StringToFile(OldFilename, OldAssetJSON) ||
+		!WriteUtf8StringToFile(NewFilename, NewAssetJSON) )
+	{
+		// we failed to write the files - we won't be able to perform a meaningful diff:
+		return EAssetCommandResult::Unhandled;
+	}
+	
+	// launch external diff process, ala UAssetDefinitionDefault::PerformAssetDiff
+	IAssetTools& AssetTools = FAssetToolsModule::GetModule().Get();
+	const FString DiffCommand = GetDefault<UEditorLoadingSavingSettings>()->TextDiffToolPath.FilePath;
+	AssetTools.CreateDiffProcess(DiffCommand, OldFilename, NewFilename);
+
+	return EAssetCommandResult::Handled;
+}
+
 UThumbnailInfo* UAssetDefinition_World::LoadThumbnailInfo(const FAssetData& InAsset) const
 {
 	return UE::Editor::FindOrCreateThumbnailInfo(InAsset.GetAsset(), UWorldThumbnailInfo::StaticClass());
diff --git a/Engine/Plugins/Editor/EngineAssetDefinitions/Source/Public/AssetDefinition_World.h b/Engine/Plugins/Editor/EngineAssetDefinitions/Source/Public/AssetDefinition_World.h
--- a/Engine/Plugins/Editor/EngineAssetDefinitions/Source/Public/AssetDefinition_World.h
+++ b/Engine/Plugins/Editor/EngineAssetDefinitions/Source/Public/AssetDefinition_World.h
@@ -28,6 +28,7 @@ public:
 
 	virtual TArray<FAssetData> PrepareToActivateAssets(const FAssetActivateArgs& ActivateArgs) const override;
 	virtual EAssetCommandResult OpenAssets(const FAssetOpenArgs& OpenArgs) const override;
+	virtual EAssetCommandResult PerformAssetDiff(const FAssetDiffArgs& DiffArgs) const override;
 	
 	virtual FAssetSupportResponse CanRename(const FAssetData& InAsset) const override;
 	virtual FAssetSupportResponse CanDuplicate(const FAssetData& InAsset) const override;
