AssetDefinition: Export symbols for AssetDefintion_Blueprint and AssetDefinition_ClassTypeBase

diff --git a/Engine/Plugins/Editor/EngineAssetDefinitions/Source/Public/Script/AssetDefinition_Blueprint.h b/Engine/Plugins/Editor/EngineAssetDefinitions/Source/Public/Script/AssetDefinition_Blueprint.h
--- a/Engine/Plugins/Editor/EngineAssetDefinitions/Source/Public/Script/AssetDefinition_Blueprint.h
+++ b/Engine/Plugins/Editor/EngineAssetDefinitions/Source/Public/Script/AssetDefinition_Blueprint.h
@@ -12,7 +12,7 @@ struct FToolMenuContext;
 struct FAssetData;
 class UFactory;
 
-UCLASS()
+UCLASS(MinimalAPI)
 class UAssetDefinition_Blueprint : public UAssetDefinition_ClassTypeBase
 {
 	GENERATED_BODY()
@@ -20,8 +20,8 @@ class UAssetDefinition_Blueprint : public UAssetDefinition_ClassTypeBase
 public:
 	// UAssetDefinition Implementation
 	virtual FText GetAssetDisplayName() const override { return NSLOCTEXT("AssetTypeActions", "AssetTypeActions_Blueprint", "Blueprint Class"); }
-	virtual FText GetAssetDisplayName(const FAssetData& AssetData) const override;
-	virtual FText GetAssetDescription(const FAssetData& AssetData) const override;
+	ENGINEASSETDEFINITIONS_API virtual FText GetAssetDisplayName(const FAssetData& AssetData) const override;
+	ENGINEASSETDEFINITIONS_API virtual FText GetAssetDescription(const FAssetData& AssetData) const override;
 	virtual FLinearColor GetAssetColor() const override { return FLinearColor(FColor( 63, 126, 255 )); }
 	virtual TSoftClassPtr<UObject> GetAssetClass() const override { return UBlueprint::StaticClass(); }
 	virtual TConstArrayView<FAssetCategoryPath> GetAssetCategories() const override
@@ -31,21 +31,21 @@ public:
 	}
 
 	virtual bool CanMerge() const override { return true; }
-	virtual EAssetCommandResult Merge(const FAssetAutomaticMergeArgs& MergeArgs) const override;
-	virtual EAssetCommandResult Merge(const FAssetManualMergeArgs& MergeArgs) const override;
+	ENGINEASSETDEFINITIONS_API virtual EAssetCommandResult Merge(const FAssetAutomaticMergeArgs& MergeArgs) const override;
+	ENGINEASSETDEFINITIONS_API virtual EAssetCommandResult Merge(const FAssetManualMergeArgs& MergeArgs) const override;
 	
-	virtual EAssetCommandResult PerformAssetDiff(const FAssetDiffArgs& DiffArgs) const override;
+	ENGINEASSETDEFINITIONS_API virtual EAssetCommandResult PerformAssetDiff(const FAssetDiffArgs& DiffArgs) const override;
 
-	virtual UThumbnailInfo* LoadThumbnailInfo(const FAssetData& InAssetData) const override;
-	virtual EAssetCommandResult OpenAssets(const FAssetOpenArgs& OpenArgs) const override;
+	ENGINEASSETDEFINITIONS_API virtual UThumbnailInfo* LoadThumbnailInfo(const FAssetData& InAssetData) const override;
+	ENGINEASSETDEFINITIONS_API virtual EAssetCommandResult OpenAssets(const FAssetOpenArgs& OpenArgs) const override;
 	// UAssetDefinition End
 
 	// UAssetDefinition_ClassTypeBase Implementation
-	virtual TWeakPtr<IClassTypeActions> GetClassTypeActions(const FAssetData& AssetData) const override;
+	ENGINEASSETDEFINITIONS_API virtual TWeakPtr<IClassTypeActions> GetClassTypeActions(const FAssetData& AssetData) const override;
 	// UAssetDefinition_ClassTypeBase End
 	
 public:
-	virtual UFactory* GetFactoryForBlueprintType(UBlueprint* InBlueprint) const;
+	ENGINEASSETDEFINITIONS_API virtual UFactory* GetFactoryForBlueprintType(UBlueprint* InBlueprint) const;
 };
 
 #if UE_ENABLE_INCLUDE_ORDER_DEPRECATED_IN_5_2
diff --git a/Engine/Plugins/Editor/EngineAssetDefinitions/Source/Public/Script/AssetDefinition_ClassTypeBase.h b/Engine/Plugins/Editor/EngineAssetDefinitions/Source/Public/Script/AssetDefinition_ClassTypeBase.h
--- a/Engine/Plugins/Editor/EngineAssetDefinitions/Source/Public/Script/AssetDefinition_ClassTypeBase.h
+++ b/Engine/Plugins/Editor/EngineAssetDefinitions/Source/Public/Script/AssetDefinition_ClassTypeBase.h
@@ -9,7 +9,7 @@
 class IClassTypeActions;
 
 /** Base class for "class type" assets (C++ classes and Blueprints */
-UCLASS(Abstract)
+UCLASS(Abstract, MinimalAPI)
 class UAssetDefinition_ClassTypeBase : public UAssetDefinitionDefault
 {
 	GENERATED_BODY()
@@ -20,7 +20,7 @@ public:
 	{
 		return FAssetSupportResponse::NotSupported();
 	}
-	virtual TSharedPtr<class SWidget> GetThumbnailOverlay(const FAssetData& AssetData) const override;
+	ENGINEASSETDEFINITIONS_API virtual TSharedPtr<class SWidget> GetThumbnailOverlay(const FAssetData& AssetData) const override;
 	// UAssetDefinition End
 
 	/** Get the class type actions for this asset */
