Hide texture asset context menus showing for read-only assets
#rnx

diff --git a/Engine/Plugins/Editor/EngineAssetDefinitions/Source/Private/AssetDefinition_Texture.cpp b/Engine/Plugins/Editor/EngineAssetDefinitions/Source/Private/AssetDefinition_Texture.cpp
--- a/Engine/Plugins/Editor/EngineAssetDefinitions/Source/Private/AssetDefinition_Texture.cpp
+++ b/Engine/Plugins/Editor/EngineAssetDefinitions/Source/Private/AssetDefinition_Texture.cpp
@@ -132,13 +132,13 @@ namespace MenuExtension_Texture
 			FToolMenuOwnerScoped OwnerScoped(UE_MODULE_NAME);
 
 			{
-				UToolMenu* Menu = UE::ContentBrowser::ExtendToolMenu_AssetContextMenu(UTexture::StaticClass());
-		
-				FToolMenuSection& Section = Menu->FindOrAddSection("GetAssetActions");
-				Section.AddDynamicEntry(NAME_None, FNewToolMenuSectionDelegate::CreateLambda([](FToolMenuSection& InSection)
+				UE::ContentBrowser::ExtendToolMenu_AssetContextMenu(UTexture::StaticClass())->AddDynamicSection(NAME_None, FNewToolMenuDelegate::CreateLambda([](UToolMenu* InMenu)
 				{
-					if (const UContentBrowserAssetContextMenuContext* Context = UContentBrowserAssetContextMenuContext::FindContextWithAssets(InSection))
+					const UContentBrowserAssetContextMenuContext* Context = InMenu->FindContext<UContentBrowserAssetContextMenuContext>();
+					if (Context && Context->SelectedAssets.Num() > 0)
 					{
+						FToolMenuSection& InSection = InMenu->FindOrAddSection("GetAssetActions");
+
 						{
 							const TAttribute<FText> Label = LOCTEXT("Texture_CreateMaterial", "Create Material");
 							const TAttribute<FText> ToolTip = LOCTEXT("Texture_CreateMaterialTooltip", "Creates a new material using this texture.");
@@ -156,6 +156,7 @@ namespace MenuExtension_Texture
 							InSection.AddMenuEntry("Texture_FindMaterials", Label, ToolTip, Icon, UIAction);
 						}
 						
+						if (Context->bCanBeModified) 
 						{
 							const TAttribute<FText> Label = LOCTEXT("Texture_ResizeSourceToPowerOfTwo", "Texture Source Resize To Power of Two");
 							const TAttribute<FText> ToolTip = LOCTEXT("Texture_ResizeSourceToPowerOfTwoTooltip", "Change texture source dimensions to the nearest power of two.");
@@ -164,6 +165,7 @@ namespace MenuExtension_Texture
 							InSection.AddMenuEntry("Texture_ResizeSourceToPowerOfTwo", Label, ToolTip, Icon, UIAction);
 						}
 
+						if (Context->bCanBeModified)
 						{
 							const TAttribute<FText> Label = LOCTEXT("Texture_ResizeSource", "Texture Source Reduce Size");
 							const TAttribute<FText> ToolTip = LOCTEXT("Texture_ResizeSourceTooltip", "Reduce texture asset size by shrinking the texture source dimensions.");
@@ -172,6 +174,7 @@ namespace MenuExtension_Texture
 							InSection.AddMenuEntry("Texture_ResizeSource", Label, ToolTip, Icon, UIAction);
 						}
 												
+						if (Context->bCanBeModified)
 						{
 							const TAttribute<FText> Label = LOCTEXT("Texture_ConvertTo8bit", "Texture Source Convert To 8 bit or minimum bit depth");
 							const TAttribute<FText> ToolTip = LOCTEXT("Texture_ConvertTo8bitTooltip", "Reduce texture asset size by converting 16/32 bit source data to 8 bit or minimum compatible bit depth.");
@@ -180,6 +183,7 @@ namespace MenuExtension_Texture
 							InSection.AddMenuEntry("Texture_ConvertTo8bit", Label, ToolTip, Icon, UIAction);
 						}
 						
+						if (Context->bCanBeModified)
 						{
 							const TAttribute<FText> Label = LOCTEXT("Texture2D_JPEGSource", "Texture Source Compress With JPEG");
 							const TAttribute<FText> ToolTip = LOCTEXT("Texture2D_JPEGSourceTooltip", "Reduce texture asset size by compressing source with JPEG.");
@@ -193,13 +197,13 @@ namespace MenuExtension_Texture
 
 			{
 				// VT actions should only be on Texture2D , not all UTexture
-				UToolMenu* Menu = UE::ContentBrowser::ExtendToolMenu_AssetContextMenu(UTexture2D::StaticClass());
-		
-				FToolMenuSection& Section = Menu->FindOrAddSection("GetAssetActions");
-				Section.AddDynamicEntry(NAME_None, FNewToolMenuSectionDelegate::CreateLambda([](FToolMenuSection& InSection)
+				UE::ContentBrowser::ExtendToolMenu_AssetContextMenu(UTexture2D::StaticClass())->AddDynamicSection(NAME_None, FNewToolMenuDelegate::CreateLambda([](UToolMenu* InMenu)
 				{
-					if (const UContentBrowserAssetContextMenuContext* Context = UContentBrowserAssetContextMenuContext::FindContextWithAssets(InSection))
+					const UContentBrowserAssetContextMenuContext* Context = InMenu->FindContext<UContentBrowserAssetContextMenuContext>();
+					if (Context && Context->bCanBeModified && Context->SelectedAssets.Num() > 0)
 					{											
+						FToolMenuSection& InSection = InMenu->FindOrAddSection("GetAssetActions");
+
 						static const auto CVarVirtualTexturesEnabled = IConsoleManager::Get().FindTConsoleVariableDataInt(TEXT("r.VirtualTextures"));
 						check(CVarVirtualTexturesEnabled);
 						
