Hide more texture asset context menus showing for read-only assets (they crash when clicked on cooked assets as well)
Create Texture Array, Create Volume Texture
#rnx

diff --git a/Engine/Plugins/Editor/EngineAssetDefinitions/Source/Private/AssetDefinition_Texture2D.cpp b/Engine/Plugins/Editor/EngineAssetDefinitions/Source/Private/AssetDefinition_Texture2D.cpp
--- a/Engine/Plugins/Editor/EngineAssetDefinitions/Source/Private/AssetDefinition_Texture2D.cpp
+++ b/Engine/Plugins/Editor/EngineAssetDefinitions/Source/Private/AssetDefinition_Texture2D.cpp
@@ -113,46 +113,47 @@ namespace MenuExtension_Texture2D
 		UToolMenus::RegisterStartupCallback(FSimpleMulticastDelegate::FDelegate::CreateLambda([]()
 		{
 			FToolMenuOwnerScoped OwnerScoped(UE_MODULE_NAME);
-			UToolMenu* Menu = UE::ContentBrowser::ExtendToolMenu_AssetContextMenu(UTexture2D::StaticClass());
-	        
-			FToolMenuSection& Section = Menu->FindOrAddSection("GetAssetActions");
-			Section.AddDynamicEntry(NAME_None, FNewToolMenuSectionDelegate::CreateLambda([](FToolMenuSection& InSection)
+			UE::ContentBrowser::ExtendToolMenu_AssetContextMenu(UTexture2D::StaticClass())->AddDynamicSection(NAME_None, FNewToolMenuDelegate::CreateLambda([](UToolMenu* InMenu)
 			{
-				const UContentBrowserAssetContextMenuContext* Context = UContentBrowserAssetContextMenuContext::FindContextWithAssets(InSection);
-				TSharedRef<FPathPermissionList> ClassPathPermissionList = IAssetTools::Get().GetAssetClassPathPermissionList(EAssetClassAction::CreateAsset);
-
-				if (ClassPathPermissionList->PassesFilter(USlateBrushAsset::StaticClass()->GetClassPathName().ToString()))
+				const UContentBrowserAssetContextMenuContext* Context = InMenu->FindContext<UContentBrowserAssetContextMenuContext>();
+				if (Context && Context->bCanBeModified && Context->SelectedAssets.Num() > 0)
 				{
-					const TAttribute<FText> Label = LOCTEXT("Texture2D_CreateSlateBrush", "Create Slate Brush");
-					const TAttribute<FText> ToolTip = LOCTEXT("Texture2D_CreateSlateBrushToolTip", "Creates a new slate brush using this texture.");
-					const FSlateIcon Icon = FSlateIcon(FAppStyle::GetAppStyleSetName(), "ClassIcon.SlateBrushAsset");
-					const FToolMenuExecuteAction UIAction = FToolMenuExecuteAction::CreateStatic(&ExecuteCreateSlateBrush);
-
-					InSection.AddMenuEntry("Texture2D_CreateSlateBrush", Label, ToolTip, Icon, UIAction);
-				}
+					TSharedRef<FPathPermissionList> ClassPathPermissionList = IAssetTools::Get().GetAssetClassPathPermissionList(EAssetClassAction::CreateAsset);
+					FToolMenuSection& InSection = InMenu->FindOrAddSection("GetAssetActions");
 
-				if (ClassPathPermissionList->PassesFilter(UTexture2DArray::StaticClass()->GetClassPathName().ToString()))
-				{
-					static const auto AllowTextureArrayAssetCreationVar = IConsoleManager::Get().FindTConsoleVariableDataInt(TEXT("r.AllowTexture2DArrayCreation"));
-					if (AllowTextureArrayAssetCreationVar->GetValueOnGameThread() != 0)
+					if (ClassPathPermissionList->PassesFilter(USlateBrushAsset::StaticClass()->GetClassPathName().ToString()))
 					{
-						const TAttribute<FText> Label = LOCTEXT("Texture_Texture2DArray", "Create Texture Array");
-						const TAttribute<FText> ToolTip = LOCTEXT("Texture_CreateTexture2DArrayTooltip", "Creates a new texture array.");
-						const FSlateIcon Icon = FSlateIcon(FAppStyle::GetAppStyleSetName(), "ClassIcon.Texture2D");
-						const FToolMenuExecuteAction UIAction = FToolMenuExecuteAction::CreateStatic(&ExecuteCreateTextureArray);
+						const TAttribute<FText> Label = LOCTEXT("Texture2D_CreateSlateBrush", "Create Slate Brush");
+						const TAttribute<FText> ToolTip = LOCTEXT("Texture2D_CreateSlateBrushToolTip", "Creates a new slate brush using this texture.");
+						const FSlateIcon Icon = FSlateIcon(FAppStyle::GetAppStyleSetName(), "ClassIcon.SlateBrushAsset");
+						const FToolMenuExecuteAction UIAction = FToolMenuExecuteAction::CreateStatic(&ExecuteCreateSlateBrush);
 
-						InSection.AddMenuEntry("Texture_Texture2DArray", Label, ToolTip, Icon, UIAction);
+						InSection.AddMenuEntry("Texture2D_CreateSlateBrush", Label, ToolTip, Icon, UIAction);
 					}
-				}
 
-				if (Context->SelectedAssets.Num() == 1 && ClassPathPermissionList->PassesFilter(UVolumeTexture::StaticClass()->GetClassPathName().ToString()))
-                {
-					const TAttribute<FText> Label = LOCTEXT("Texture2D_CreateVolumeTexture", "Create Volume Texture");
-					const TAttribute<FText> ToolTip = LOCTEXT("Texture2D_CreateVolumeTextureToolTip", "Creates a new volume texture using this texture.");
-					const FSlateIcon Icon = FSlateIcon(FAppStyle::GetAppStyleSetName(), "ContentBrowser.AssetActions.VolumeTexture");
-					const FToolMenuExecuteAction UIAction = FToolMenuExecuteAction::CreateStatic(&ExecuteCreateVolumeTexture);
+					if (ClassPathPermissionList->PassesFilter(UTexture2DArray::StaticClass()->GetClassPathName().ToString()))
+					{
+						static const auto AllowTextureArrayAssetCreationVar = IConsoleManager::Get().FindTConsoleVariableDataInt(TEXT("r.AllowTexture2DArrayCreation"));
+						if (AllowTextureArrayAssetCreationVar->GetValueOnGameThread() != 0)
+						{
+							const TAttribute<FText> Label = LOCTEXT("Texture_Texture2DArray", "Create Texture Array");
+							const TAttribute<FText> ToolTip = LOCTEXT("Texture_CreateTexture2DArrayTooltip", "Creates a new texture array.");
+							const FSlateIcon Icon = FSlateIcon(FAppStyle::GetAppStyleSetName(), "ClassIcon.Texture2D");
+							const FToolMenuExecuteAction UIAction = FToolMenuExecuteAction::CreateStatic(&ExecuteCreateTextureArray);
+
+							InSection.AddMenuEntry("Texture_Texture2DArray", Label, ToolTip, Icon, UIAction);
+						}
+					}
 
-					InSection.AddMenuEntry("Texture2D_CreateVolumeTexture", Label, ToolTip, Icon, UIAction);
+					if (Context->SelectedAssets.Num() == 1 && ClassPathPermissionList->PassesFilter(UVolumeTexture::StaticClass()->GetClassPathName().ToString()))
+					{
+						const TAttribute<FText> Label = LOCTEXT("Texture2D_CreateVolumeTexture", "Create Volume Texture");
+						const TAttribute<FText> ToolTip = LOCTEXT("Texture2D_CreateVolumeTextureToolTip", "Creates a new volume texture using this texture.");
+						const FSlateIcon Icon = FSlateIcon(FAppStyle::GetAppStyleSetName(), "ContentBrowser.AssetActions.VolumeTexture");
+						const FToolMenuExecuteAction UIAction = FToolMenuExecuteAction::CreateStatic(&ExecuteCreateVolumeTexture);
+
+						InSection.AddMenuEntry("Texture2D_CreateVolumeTexture", Label, ToolTip, Icon, UIAction);
+					}
 				}
 			}));
 		}));
