Fixing SkeletalMesh and Niagara asset context menus to respect read-only assets
#rnx

diff --git a/Engine/Plugins/Editor/EngineAssetDefinitions/Source/Private/Animation/AssetDefinition_SkeletalMesh.cpp b/Engine/Plugins/Editor/EngineAssetDefinitions/Source/Private/Animation/AssetDefinition_SkeletalMesh.cpp
--- a/Engine/Plugins/Editor/EngineAssetDefinitions/Source/Private/Animation/AssetDefinition_SkeletalMesh.cpp
+++ b/Engine/Plugins/Editor/EngineAssetDefinitions/Source/Private/Animation/AssetDefinition_SkeletalMesh.cpp
@@ -533,52 +533,56 @@ namespace MenuExtension_SkeletalMesh
 			FToolMenuSection& Section = Menu->FindOrAddSection("GetAssetActions");
 			Section.AddDynamicEntry(NAME_None, FNewToolMenuSectionDelegate::CreateLambda([](FToolMenuSection& InSection)
 			{
+				const UContentBrowserAssetContextMenuContext* Context = UContentBrowserAssetContextMenuContext::FindContextWithAssets(InSection);
+				if (Context && Context->bCanBeModified)
 				{
-					InSection.AddSubMenu(
-						"CreateSkeletalMeshSubmenu",
-						LOCTEXT("CreateSkeletalMeshSubmenu", "Create"),
-						LOCTEXT("CreateSkeletalMeshSubmenu_ToolTip", "Create related assets"),
-						FNewToolMenuDelegate::CreateStatic(&FillCreateMenu),
-						false,
-						FSlateIcon(FAppStyle::GetAppStyleSetName(), "Persona.AssetActions.CreateAnimAsset")
-					);
-				}
+					{
+						InSection.AddSubMenu(
+							"CreateSkeletalMeshSubmenu",
+							LOCTEXT("CreateSkeletalMeshSubmenu", "Create"),
+							LOCTEXT("CreateSkeletalMeshSubmenu_ToolTip", "Create related assets"),
+							FNewToolMenuDelegate::CreateStatic(&FillCreateMenu),
+							false,
+							FSlateIcon(FAppStyle::GetAppStyleSetName(), "Persona.AssetActions.CreateAnimAsset")
+						);
+					}
 
-				static const auto AllowSkinnedMeshes = IConsoleManager::Get().FindTConsoleVariableDataInt(TEXT("r.Nanite.AllowSkinnedMeshes"));
-				static const bool bAllowSkinnedMeshes = (AllowSkinnedMeshes && AllowSkinnedMeshes->GetValueOnAnyThread() != 0);
-				if (bAllowSkinnedMeshes)
-				{
-					InSection.AddSubMenu(
-						"SkeletalMesh_NaniteSubmenu",
-						LOCTEXT("SkeletalMesh_NaniteMenu", "Nanite"),
-						LOCTEXT("SkeletalMesh_NaniteTooltip", "Nanite Options and Tools"),
-						FNewToolMenuDelegate::CreateStatic(&GetNaniteMenu),
-						false,
-						FSlateIcon(FAppStyle::GetAppStyleSetName(), "Icons.Adjust")
-					);
-				}
+					static const auto AllowSkinnedMeshes = IConsoleManager::Get().FindTConsoleVariableDataInt(TEXT("r.Nanite.AllowSkinnedMeshes"));
+					static const bool bAllowSkinnedMeshes = (AllowSkinnedMeshes && AllowSkinnedMeshes->GetValueOnAnyThread() != 0);
+					if (bAllowSkinnedMeshes)
+					{
+						InSection.AddSubMenu(
+							"SkeletalMesh_NaniteSubmenu",
+							LOCTEXT("SkeletalMesh_NaniteMenu", "Nanite"),
+							LOCTEXT("SkeletalMesh_NaniteTooltip", "Nanite Options and Tools"),
+							FNewToolMenuDelegate::CreateStatic(&GetNaniteMenu),
+							false,
+							FSlateIcon(FAppStyle::GetAppStyleSetName(), "Icons.Adjust")
+						);
+					}
 
-				if (UContentBrowserAssetContextMenuContext::GetNumAssetsSelected(InSection))
-				{
-					InSection.AddSubMenu(
-						"SkeletalMesh_LODImport",	
-						LOCTEXT("SkeletalMesh_LODImport", "Import LOD"),
-						LOCTEXT("SkeletalMesh_LODImportTooltip", "Select which LODs to import."),
-						FNewToolMenuDelegate::CreateStatic(&GetLODMenu),
-						false,
-						FSlateIcon(FAppStyle::GetAppStyleSetName(), "Icons.LOD")
-					);
-				}
+					if (UContentBrowserAssetContextMenuContext::GetNumAssetsSelected(InSection))
+					{
+						InSection.AddSubMenu(
+							"SkeletalMesh_LODImport",	
+							LOCTEXT("SkeletalMesh_LODImport", "Import LOD"),
+							LOCTEXT("SkeletalMesh_LODImportTooltip", "Select which LODs to import."),
+							FNewToolMenuDelegate::CreateStatic(&GetLODMenu),
+							false,
+							FSlateIcon(FAppStyle::GetAppStyleSetName(), "Icons.LOD")
+						);
+					}
 
-				{
-					InSection.AddSubMenu(
-						"SkeletonSubmenu",
-						LOCTEXT("SkeletonSubmenu", "Skeleton"),
-						LOCTEXT("SkeletonSubmenu_ToolTip", "Skeleton related actions"),
-						FNewToolMenuDelegate::CreateStatic(&FillSkeletonMenu),
-						false,
-						FSlateIcon(FAppStyle::GetAppStyleSetName(), "Icons.SkeletalMesh")
-					);
+					{
+						InSection.AddSubMenu(
+							"SkeletonSubmenu",
+							LOCTEXT("SkeletonSubmenu", "Skeleton"),
+							LOCTEXT("SkeletonSubmenu_ToolTip", "Skeleton related actions"),
+							FNewToolMenuDelegate::CreateStatic(&FillSkeletonMenu),
+							false,
+							FSlateIcon(FAppStyle::GetAppStyleSetName(), "Icons.SkeletalMesh")
+						);
+					}
 				}
 			}));
 		}));
diff --git a/Engine/Plugins/FX/Niagara/Source/NiagaraEditor/Private/NiagaraEditorMenuHelpers.cpp b/Engine/Plugins/FX/Niagara/Source/NiagaraEditor/Private/NiagaraEditorMenuHelpers.cpp
--- a/Engine/Plugins/FX/Niagara/Source/NiagaraEditor/Private/NiagaraEditorMenuHelpers.cpp
+++ b/Engine/Plugins/FX/Niagara/Source/NiagaraEditor/Private/NiagaraEditorMenuHelpers.cpp
@@ -450,7 +450,14 @@ void NiagaraEditorMenuHelpers::RegisterNiagaraAssetTagMenusForAssets()
 		// Then we extend the asset context menu by adding an empty sub menu
 		const FName SubMenuName = TEXT("ManageTags");
 		UToolMenu* AssetContextMenu = UE::ContentBrowser::ExtendToolMenu_AssetContextMenu(ExtendedClass);
-		AssetContextMenu->AddSubMenu(OwnerScoped.GetOwner(), "GetAssetActions", SubMenuName, LOCTEXT("ManageTabsSubMenuLabel", "Manage Tags"));
+		AssetContextMenu->AddDynamicSection(NAME_None, FNewToolMenuDelegate::CreateLambda([SubMenuName](UToolMenu* InMenu)
+		{
+			const UContentBrowserAssetContextMenuContext* Context = InMenu->FindContext<UContentBrowserAssetContextMenuContext>();
+			if (Context && Context->bCanBeModified)
+			{
+				InMenu->AddSubMenu(FToolMenuOwner(), "GetAssetActions", SubMenuName, LOCTEXT("ManageTabsSubMenuLabel", "Manage Tags"));
+			}
+		}));
 		
 		// Now that we have added a submenu with a specific name, we need the full asset context menu path including the same sub menu name
 		const FName AssetContextMenuNameWithSubMenu = UToolMenus::JoinMenuPaths(AssetContextMenu->GetMenuName(), SubMenuName);
