Exporting AssetDefinition_DataTable and DataAssetFactory via API macro, so that user projects can create subclasses.

diff --git a/Engine/Plugins/Editor/EngineAssetDefinitions/Source/Private/Table/AssetDefinition_DataTable.cpp b/Engine/Plugins/Editor/EngineAssetDefinitions/Source/Private/Table/AssetDefinition_DataTable.cpp
--- a/Engine/Plugins/Editor/EngineAssetDefinitions/Source/Private/Table/AssetDefinition_DataTable.cpp
+++ b/Engine/Plugins/Editor/EngineAssetDefinitions/Source/Private/Table/AssetDefinition_DataTable.cpp
@@ -1,6 +1,6 @@
 // Copyright Epic Games, Inc. All Rights Reserved.
 
-#include "AssetDefinition_DataTable.h"
+#include "Table/AssetDefinition_DataTable.h"
 
 #include "AssetDefinitionRegistry.h"
 #include "IAssetTools.h"
diff --git a/Engine/Plugins/Editor/EngineAssetDefinitions/Source/Private/Table/AssetDefinition_DataTable.h b/Engine/Plugins/Editor/EngineAssetDefinitions/Source/Private/Table/AssetDefinition_DataTable.h
deleted file mode 100644
diff --git a/Engine/Plugins/Editor/EngineAssetDefinitions/Source/Public/Table/AssetDefinition_DataTable.h b/Engine/Plugins/Editor/EngineAssetDefinitions/Source/Public/Table/AssetDefinition_DataTable.h
new file mode 100644
--- /dev/null
+++ b/Engine/Plugins/Editor/EngineAssetDefinitions/Source/Public/Table/AssetDefinition_DataTable.h
@@ -0,0 +1,30 @@
+// Copyright Epic Games, Inc. All Rights Reserved.
+
+#pragma once
+
+#include "Engine/DataTable.h"
+#include "AssetDefinitionDefault.h"
+
+#include "AssetDefinition_DataTable.generated.h"
+
+UCLASS()
+class ENGINEASSETDEFINITIONS_API UAssetDefinition_DataTable : public UAssetDefinitionDefault
+{
+	GENERATED_BODY()
+
+public:
+	// UAssetDefinition Begin
+	virtual FText GetAssetDisplayName() const override { return NSLOCTEXT("AssetTypeActions", "AssetTypeActions_DataTable", "Data Table"); }
+	virtual FText GetAssetDisplayName(const FAssetData& AssetData) const override;
+	virtual FLinearColor GetAssetColor() const override { return FLinearColor(FColor(62, 140, 35)); }
+	virtual TSoftClassPtr<UObject> GetAssetClass() const override { return UDataTable::StaticClass(); }
+	virtual TConstArrayView<FAssetCategoryPath> GetAssetCategories() const override
+	{
+		static const auto Categories = { EAssetCategoryPaths::Misc };
+		return Categories;
+	}
+	virtual bool CanImport() const override { return true; }
+	virtual EAssetCommandResult PerformAssetDiff(const FAssetDiffArgs& DiffArgs) const override;
+	virtual EAssetCommandResult OpenAssets(const FAssetOpenArgs& OpenArgs) const override;
+	// UAssetDefinition End
+};
diff --git a/Engine/Source/Editor/UnrealEd/Classes/Factories/DataAssetFactory.h b/Engine/Source/Editor/UnrealEd/Classes/Factories/DataAssetFactory.h
--- a/Engine/Source/Editor/UnrealEd/Classes/Factories/DataAssetFactory.h
+++ b/Engine/Source/Editor/UnrealEd/Classes/Factories/DataAssetFactory.h
@@ -19,8 +19,8 @@ class UDataAssetFactory : public UFactory
 	TSubclassOf<UDataAsset> DataAssetClass;
 
 	// UFactory interface
-	virtual bool ConfigureProperties() override;
-	virtual UObject* FactoryCreateNew(UClass* Class, UObject* InParent, FName Name, EObjectFlags Flags, UObject* Context, FFeedbackContext* Warn) override;
+	UNREALED_API virtual bool ConfigureProperties() override;
+	UNREALED_API virtual UObject* FactoryCreateNew(UClass* Class, UObject* InParent, FName Name, EObjectFlags Flags, UObject* Context, FFeedbackContext* Warn) override;
 	// End of UFactory interface
 };
 
