- Extend SVT asset actions to allow creating a simple material or a material instance based off the engine default

diff --git a/Engine/Plugins/Editor/EngineAssetDefinitions/Source/Private/AssetDefinition_SparseVolumeTexture.cpp b/Engine/Plugins/Editor/EngineAssetDefinitions/Source/Private/AssetDefinition_SparseVolumeTexture.cpp
--- a/Engine/Plugins/Editor/EngineAssetDefinitions/Source/Private/AssetDefinition_SparseVolumeTexture.cpp
+++ b/Engine/Plugins/Editor/EngineAssetDefinitions/Source/Private/AssetDefinition_SparseVolumeTexture.cpp
@@ -2,10 +2,105 @@
 
 #include "AssetDefinition_SparseVolumeTexture.h"
 
+#include "Factories/SparseVolumeTextureMaterialFactory.h"
+#include "ContentBrowserMenuContexts.h"
+#include "IAssetTools.h"
+#include "ToolMenu.h"
+#include "ToolMenus.h"
+
 #define LOCTEXT_NAMESPACE "UAssetDefinition_SparseVolumeTexture"
 
 namespace MenuExtension_SparseVolumeTexture
 {
+	static void ExecuteCreateMaterial(const FToolMenuContext& InContext)
+	{
+		const UContentBrowserAssetContextMenuContext* CBContext = UContentBrowserAssetContextMenuContext::FindContextWithAssets(InContext);
+	
+		IAssetTools::Get().CreateAssetsFrom<USparseVolumeTexture>(
+			CBContext->LoadSelectedObjects<USparseVolumeTexture>(),
+			GetDefault<USparseVolumeTextureMaterialFactoryNew>()->GetSupportedClass(),
+			TEXT("_Mat"),
+			[](USparseVolumeTexture* SourceObject)
+			{
+				USparseVolumeTextureMaterialFactoryNew* Factory = NewObject<USparseVolumeTextureMaterialFactoryNew>();
+				Factory->InitialTexture = SourceObject;
+				return Factory;
+			}
+		);
+	}
+
+	static void ExecuteCreateMaterialInstance(const FToolMenuContext& InContext)
+	{
+		const UContentBrowserAssetContextMenuContext* CBContext = UContentBrowserAssetContextMenuContext::FindContextWithAssets(InContext);
+	
+		IAssetTools::Get().CreateAssetsFrom<USparseVolumeTexture>(
+			CBContext->LoadSelectedObjects<USparseVolumeTexture>(),
+			GetDefault<USparseVolumeTextureMaterialInstanceFactoryNew>()->GetSupportedClass(),
+			TEXT("_MIC"),
+			[](USparseVolumeTexture* SourceObject)
+			{
+				USparseVolumeTextureMaterialInstanceFactoryNew* Factory = NewObject<USparseVolumeTextureMaterialInstanceFactoryNew>();
+				Factory->InitialTexture = SourceObject;
+				return Factory;
+			}
+		);
+	}
+
+	static void ExtendAssetActions()
+	{
+		FToolMenuOwnerScoped OwnerScoped(UE_MODULE_NAME);
+		{
+			UE::ContentBrowser::ExtendToolMenu_AssetContextMenu(USparseVolumeTexture::StaticClass())->AddDynamicSection(
+				NAME_None,
+				FNewToolMenuDelegate::CreateLambda(
+					[](UToolMenu* InMenu)
+					{
+						const UContentBrowserAssetContextMenuContext* Context = InMenu->FindContext<UContentBrowserAssetContextMenuContext>();
+						if (Context && Context->SelectedAssets.Num() > 0)
+						{
+							FToolMenuSection& InSection = InMenu->FindOrAddSection("GetAssetActions");
+							{
+								const TAttribute<FText> Label = LOCTEXT("Texture_CreateMaterial", "Create Material");
+								const TAttribute<FText> ToolTip = LOCTEXT("Texture_CreateMaterialTooltip", "Creates a new material using this sparse volume texture.");
+								const FSlateIcon Icon = FSlateIcon(FAppStyle::GetAppStyleSetName(), "ClassIcon.Material");
+								const FToolMenuExecuteAction UIAction = FToolMenuExecuteAction::CreateStatic(&ExecuteCreateMaterial);
+								InSection.AddMenuEntry("Texture_CreateMaterial", Label, ToolTip, Icon, UIAction);
+							}
+						}
+					}
+				)
+			);
+
+			UE::ContentBrowser::ExtendToolMenu_AssetContextMenu(USparseVolumeTexture::StaticClass())->AddDynamicSection(
+				NAME_None,
+				FNewToolMenuDelegate::CreateLambda(
+					[](UToolMenu* InMenu)
+					{
+						const UContentBrowserAssetContextMenuContext* Context = InMenu->FindContext<UContentBrowserAssetContextMenuContext>();
+						if (Context && Context->SelectedAssets.Num() > 0)
+						{
+							FToolMenuSection& InSection = InMenu->FindOrAddSection("GetAssetActions");
+							{
+								const TAttribute<FText> Label = LOCTEXT("Texture_CreateMaterialInstance", "Create Material Instance");
+								const TAttribute<FText> ToolTip = LOCTEXT("Texture_CreateMaterialInstanceTooltip", "Creates a new material instance using this sparse volume texture.");
+								const FSlateIcon Icon = FSlateIcon(FAppStyle::GetAppStyleSetName(), "ClassIcon.Material");
+								const FToolMenuExecuteAction UIAction = FToolMenuExecuteAction::CreateStatic(&ExecuteCreateMaterialInstance);
+								InSection.AddMenuEntry("Texture_CreateMaterialInstance", Label, ToolTip, Icon, UIAction);
+							}
+						}
+					}
+				)
+			);
+		}
+	}
+
+	static FDelayedAutoRegisterHelper DelayedAutoRegister(
+		EDelayedRegisterRunPhase::EndOfEngineInit,
+		[]
+		{
+			UToolMenus::RegisterStartupCallback(FSimpleMulticastDelegate::FDelegate::CreateStatic(&ExtendAssetActions));
+		}
+	);
 }
 
 #undef LOCTEXT_NAMESPACE
diff --git a/Engine/Plugins/Editor/EngineAssetDefinitions/Source/Private/AssetDefinition_SparseVolumeTexture.h b/Engine/Plugins/Editor/EngineAssetDefinitions/Source/Private/AssetDefinition_SparseVolumeTexture.h
--- a/Engine/Plugins/Editor/EngineAssetDefinitions/Source/Private/AssetDefinition_SparseVolumeTexture.h
+++ b/Engine/Plugins/Editor/EngineAssetDefinitions/Source/Private/AssetDefinition_SparseVolumeTexture.h
@@ -4,6 +4,7 @@
 
 #include "SparseVolumeTexture/SparseVolumeTexture.h"
 #include "AssetDefinitionDefault.h"
+#include "Factories/Factory.h"
 
 #include "AssetDefinition_SparseVolumeTexture.generated.h"
 
diff --git a/Engine/Source/Editor/UnrealEd/Classes/Factories/SparseVolumeTextureMaterialFactory.h b/Engine/Source/Editor/UnrealEd/Classes/Factories/SparseVolumeTextureMaterialFactory.h
new file mode 100644
--- /dev/null
+++ b/Engine/Source/Editor/UnrealEd/Classes/Factories/SparseVolumeTextureMaterialFactory.h
@@ -0,0 +1,46 @@
+// Copyright Epic Games, Inc. All Rights Reserved.
+
+#pragma once
+
+#if WITH_EDITOR
+
+#include "CoreMinimal.h"
+#include "Factories/Factory.h"
+#include "EditorReimportHandler.h"
+#include "MaterialInstanceConstantFactoryNew.h"
+
+#include "SparseVolumeTextureMaterialFactory.generated.h"
+
+class UMaterialInterface;
+class USparseVolumeTexture;
+
+UCLASS(MinimalAPI)
+class USparseVolumeTextureMaterialFactoryNew : public UFactory
+{
+	GENERATED_UCLASS_BODY()
+
+	UPROPERTY()
+	TObjectPtr<USparseVolumeTexture> InitialTexture;
+
+	//~ Begin UFactory Interface
+	virtual UObject* FactoryCreateNew(UClass* Class, UObject* InParent, FName Name, EObjectFlags Flags, UObject* Context, FFeedbackContext* Warn) override;
+	//~ Begin UFactory Interface	
+};
+
+UCLASS(config=Engine, MinimalAPI)
+class USparseVolumeTextureMaterialInstanceFactoryNew : public UMaterialInstanceConstantFactoryNew
+{
+	GENERATED_UCLASS_BODY()
+
+	UPROPERTY()
+	TObjectPtr<USparseVolumeTexture> InitialTexture;
+
+	UPROPERTY(config)
+	TSoftObjectPtr<UMaterialInterface> DefaultSVTMaterial;
+
+	//~ Begin UFactory Interface
+	virtual UObject* FactoryCreateNew(UClass* Class, UObject* InParent, FName Name, EObjectFlags Flags, UObject* Context, FFeedbackContext* Warn) override;
+	//~ Begin UFactory Interface	
+};
+
+#endif // WITH_EDITOR
diff --git a/Engine/Source/Editor/UnrealEd/Private/Factories/SparseVolumeTextureMaterialFactory.cpp b/Engine/Source/Editor/UnrealEd/Private/Factories/SparseVolumeTextureMaterialFactory.cpp
new file mode 100644
--- /dev/null
+++ b/Engine/Source/Editor/UnrealEd/Private/Factories/SparseVolumeTextureMaterialFactory.cpp
@@ -0,0 +1,92 @@
+// Copyright Epic Games, Inc. All Rights Reserved.
+
+#include "Factories/SparseVolumeTextureMaterialFactory.h"
+
+#include "Engine/RendererSettings.h"
+#include "Materials/MaterialExpressionSparseVolumeTextureSample.h"
+#include "Materials/MaterialInstance.h"
+#include "Materials/MaterialInstanceConstant.h"
+#include "Materials/Material.h"
+#include "MaterialDomain.h"
+
+#if WITH_EDITOR
+
+USparseVolumeTextureMaterialFactoryNew::USparseVolumeTextureMaterialFactoryNew(const FObjectInitializer& ObjectInitializer)
+	: Super(ObjectInitializer)
+{
+	SupportedClass = UMaterial::StaticClass();
+	bCreateNew = false;
+	bEditAfterNew = true;
+}
+
+UObject* USparseVolumeTextureMaterialFactoryNew::FactoryCreateNew(UClass* Class, UObject* InParent, FName Name, EObjectFlags Flags, UObject* Context, FFeedbackContext* Warn)
+{
+	UMaterial* NewMaterial = NewObject<UMaterial>(InParent, Class, Name, Flags);
+	NewMaterial->MaterialDomain = EMaterialDomain::MD_Volume;
+	NewMaterial->BlendMode = BLEND_Additive;
+
+	UMaterialExpressionSparseVolumeTextureSample* TextureSampler = NewObject<UMaterialExpressionSparseVolumeTextureSample>(NewMaterial);
+	{
+		TextureSampler->SparseVolumeTexture = InitialTexture;
+	}
+
+	NewMaterial->GetExpressionCollection().AddExpression(TextureSampler);
+
+	UMaterialEditorOnlyData* NewMaterialEditorOnly = NewMaterial->GetEditorOnlyData();
+	{
+		FExpressionOutput& Output = TextureSampler->GetOutputs()[0];
+		FExpressionInput& Input = NewMaterialEditorOnly->EmissiveColor;
+		Input.Expression = TextureSampler;
+		Input.Mask = Output.Mask;
+		Input.MaskR = Output.MaskR;
+		Input.MaskG = Output.MaskG;
+		Input.MaskB = Output.MaskB;
+		Input.MaskA = Output.MaskA;
+	}
+	{
+		FExpressionOutput& Output = TextureSampler->GetOutputs()[1];
+		FExpressionInput& Input = NewMaterialEditorOnly->SubsurfaceColor;
+		Input.Expression = TextureSampler;
+		Input.Mask = Output.Mask;
+		Input.MaskR = Output.MaskR;
+		Input.MaskG = Output.MaskG;
+		Input.MaskB = Output.MaskB;
+		Input.MaskA = Output.MaskA;
+	}
+
+	NewMaterial->PostEditChange();
+	
+	NewMaterial->bAutomaticallySetUsageInEditor = GetDefault<URendererSettings>()->bAutomaticallySetMaterialUsageInEditorDefault;
+
+	return NewMaterial;
+}
+
+
+USparseVolumeTextureMaterialInstanceFactoryNew::USparseVolumeTextureMaterialInstanceFactoryNew(const FObjectInitializer& ObjectInitializer)
+	: Super(ObjectInitializer)
+	, DefaultSVTMaterial(FSoftObjectPath(TEXT("/Engine/EngineMaterials/SparseVolumeMaterial.SparseVolumeMaterial")))
+{
+	SupportedClass = UMaterialInstanceConstant::StaticClass();
+	bCreateNew = false;
+	bEditAfterNew = true;
+}
+
+UObject* USparseVolumeTextureMaterialInstanceFactoryNew::FactoryCreateNew(UClass* Class, UObject* InParent, FName Name, EObjectFlags Flags, UObject* Context, FFeedbackContext* Warn)
+{
+	if (InitialParent == nullptr)
+	{
+		InitialParent = DefaultSVTMaterial.LoadSynchronous();
+	}
+
+	UMaterialInstanceConstant* MIC = Cast<UMaterialInstanceConstant>(UMaterialInstanceConstantFactoryNew::FactoryCreateNew(Class, InParent, Name, Flags, Context, Warn));
+	if (MIC && InitialTexture)
+	{
+		MIC->SetSparseVolumeTextureParameterValueEditorOnly(FMaterialParameterInfo("SparseVolumeTexture"), InitialTexture);
+	}
+
+	return MIC;
+}
+
+#endif //WITH_EDITOR
+
+#undef LOCTEXT_NAMESPACE
