Put texture source actions in a submenu

#rnx

diff --git a/Engine/Plugins/Editor/EngineAssetDefinitions/Source/Private/AssetDefinition_Texture.cpp b/Engine/Plugins/Editor/EngineAssetDefinitions/Source/Private/AssetDefinition_Texture.cpp
--- a/Engine/Plugins/Editor/EngineAssetDefinitions/Source/Private/AssetDefinition_Texture.cpp
+++ b/Engine/Plugins/Editor/EngineAssetDefinitions/Source/Private/AssetDefinition_Texture.cpp
@@ -56,33 +56,33 @@ namespace MenuExtension_Texture
 		);
 	}
 
-	static void ExecuteResizeTextureSource(const FToolMenuContext& InContext)
+	static void ExecuteResizeTextureSource(const UContentBrowserAssetContextMenuContext* CBContext)
 	{
-		if (const UContentBrowserAssetContextMenuContext* CBContext = UContentBrowserAssetContextMenuContext::FindContextWithAssets(InContext))
+		if (CBContext)
 		{
 			UE::TextureAssetActions::TextureSource_Resize_WithDialog(CBContext->LoadSelectedObjects<UTexture>());
 		}
 	}
 	
-	static void ExecuteResizeToPowerOfTwoTextureSource(const FToolMenuContext& InContext)
+	static void ExecuteResizeToPowerOfTwoTextureSource(const UContentBrowserAssetContextMenuContext* CBContext)
 	{
-		if (const UContentBrowserAssetContextMenuContext* CBContext = UContentBrowserAssetContextMenuContext::FindContextWithAssets(InContext))
+		if (CBContext)
 		{
 			UE::TextureAssetActions::TextureSource_ResizeToPowerOfTwo_WithDialog(CBContext->LoadSelectedObjects<UTexture>());
 		}
 	}
 
-	static void Execute8bitTextureSource(const FToolMenuContext& InContext)
+	static void Execute8bitTextureSource(const UContentBrowserAssetContextMenuContext* CBContext)
 	{
-		if (const UContentBrowserAssetContextMenuContext* CBContext = UContentBrowserAssetContextMenuContext::FindContextWithAssets(InContext))
+		if (CBContext)
 		{
 			UE::TextureAssetActions::TextureSource_ConvertTo8bit_WithDialog(CBContext->LoadSelectedObjects<UTexture>());
 		}
 	}
 	
-	static void ExecuteJPEGTextureSource(const FToolMenuContext& InContext)
+	static void ExecuteJPEGTextureSource(const UContentBrowserAssetContextMenuContext* CBContext)
 	{
-		if (const UContentBrowserAssetContextMenuContext* CBContext = UContentBrowserAssetContextMenuContext::FindContextWithAssets(InContext))
+		if (CBContext)
 		{
 			UE::TextureAssetActions::TextureSource_JPEG_WithDialog(CBContext->LoadSelectedObjects<UTexture>());
 		}
@@ -126,6 +126,42 @@ namespace MenuExtension_Texture
 		}
 	}
 
+	static TSharedRef<SWidget> MakeTextureSourceActionMenu(const UContentBrowserAssetContextMenuContext* Context)
+	{
+		FMenuBuilder MenuBuilder(true, nullptr);
+		
+		{
+			const TAttribute<FText> Label = LOCTEXT("Texture_ResizeSourceToPowerOfTwo", "Resize To Power of Two");
+			const TAttribute<FText> ToolTip = LOCTEXT("Texture_ResizeSourceToPowerOfTwoTooltip", "Change texture source dimensions to the nearest power of two.");
+			const FUIAction UIAction = FUIAction(FExecuteAction::CreateStatic(&ExecuteResizeToPowerOfTwoTextureSource, Context), FCanExecuteAction());
+			MenuBuilder.AddMenuEntry(Label, ToolTip, FSlateIcon(), UIAction);
+		}
+
+		{
+			const TAttribute<FText> Label = LOCTEXT("Texture_ResizeSource", "Reduce Size");
+			const TAttribute<FText> ToolTip = LOCTEXT("Texture_ResizeSourceTooltip", "Reduce texture asset size by shrinking the texture source dimensions.");
+			const FUIAction UIAction = FUIAction(FExecuteAction::CreateStatic(&ExecuteResizeTextureSource, Context), FCanExecuteAction());
+			MenuBuilder.AddMenuEntry(Label, ToolTip, FSlateIcon(), UIAction);
+		}
+												
+		{
+			const TAttribute<FText> Label = LOCTEXT("Texture_ConvertTo8bit", "Convert To 8 bit or minimum bit depth");
+			const TAttribute<FText> ToolTip = LOCTEXT("Texture_ConvertTo8bitTooltip", "Reduce texture asset size by converting 16/32 bit source data to 8 bit or minimum compatible bit depth.");
+			const FUIAction UIAction = FUIAction(FExecuteAction::CreateStatic(&Execute8bitTextureSource, Context), FCanExecuteAction());
+			MenuBuilder.AddMenuEntry(Label, ToolTip, FSlateIcon(), UIAction);
+		}
+						
+		{
+			const TAttribute<FText> Label = LOCTEXT("Texture2D_JPEGSource", "Compress With JPEG");
+			const TAttribute<FText> ToolTip = LOCTEXT("Texture2D_JPEGSourceTooltip", "Reduce texture asset size by compressing source with JPEG.");
+			const FUIAction UIAction = FUIAction(FExecuteAction::CreateStatic(&ExecuteJPEGTextureSource, Context), FCanExecuteAction());
+			MenuBuilder.AddMenuEntry(Label, ToolTip, FSlateIcon(), UIAction);
+		}
+
+		return MenuBuilder.MakeWidget();
+	}
+
+
 	static FDelayedAutoRegisterHelper DelayedAutoRegister(EDelayedRegisterRunPhase::EndOfEngineInit, []{ 
 		UToolMenus::RegisterStartupCallback(FSimpleMulticastDelegate::FDelegate::CreateLambda([]()
 		{
@@ -156,40 +192,16 @@ namespace MenuExtension_Texture
 							InSection.AddMenuEntry("Texture_FindMaterials", Label, ToolTip, Icon, UIAction);
 						}
 						
-						if (Context->bCanBeModified) 
-						{
-							const TAttribute<FText> Label = LOCTEXT("Texture_ResizeSourceToPowerOfTwo", "Texture Source Resize To Power of Two");
-							const TAttribute<FText> ToolTip = LOCTEXT("Texture_ResizeSourceToPowerOfTwoTooltip", "Change texture source dimensions to the nearest power of two.");
-							const FSlateIcon Icon = FSlateIcon(FAppStyle::GetAppStyleSetName(), "ClassIcon.Texture2D");
-							const FToolMenuExecuteAction UIAction = FToolMenuExecuteAction::CreateStatic(&ExecuteResizeToPowerOfTwoTextureSource);
-							InSection.AddMenuEntry("Texture_ResizeSourceToPowerOfTwo", Label, ToolTip, Icon, UIAction);
-						}
-
-						if (Context->bCanBeModified)
-						{
-							const TAttribute<FText> Label = LOCTEXT("Texture_ResizeSource", "Texture Source Reduce Size");
-							const TAttribute<FText> ToolTip = LOCTEXT("Texture_ResizeSourceTooltip", "Reduce texture asset size by shrinking the texture source dimensions.");
-							const FSlateIcon Icon = FSlateIcon(FAppStyle::GetAppStyleSetName(), "ClassIcon.Texture2D");
-							const FToolMenuExecuteAction UIAction = FToolMenuExecuteAction::CreateStatic(&ExecuteResizeTextureSource);
-							InSection.AddMenuEntry("Texture_ResizeSource", Label, ToolTip, Icon, UIAction);
-						}
-												
-						if (Context->bCanBeModified)
-						{
-							const TAttribute<FText> Label = LOCTEXT("Texture_ConvertTo8bit", "Texture Source Convert To 8 bit or minimum bit depth");
-							const TAttribute<FText> ToolTip = LOCTEXT("Texture_ConvertTo8bitTooltip", "Reduce texture asset size by converting 16/32 bit source data to 8 bit or minimum compatible bit depth.");
-							const FSlateIcon Icon = FSlateIcon(FAppStyle::GetAppStyleSetName(), "ClassIcon.Texture2D");
-							const FToolMenuExecuteAction UIAction = FToolMenuExecuteAction::CreateStatic(&Execute8bitTextureSource);
-							InSection.AddMenuEntry("Texture_ConvertTo8bit", Label, ToolTip, Icon, UIAction);
-						}
-						
 						if (Context->bCanBeModified)
 						{
-							const TAttribute<FText> Label = LOCTEXT("Texture2D_JPEGSource", "Texture Source Compress With JPEG");
-							const TAttribute<FText> ToolTip = LOCTEXT("Texture2D_JPEGSourceTooltip", "Reduce texture asset size by compressing source with JPEG.");
-							const FSlateIcon Icon = FSlateIcon(FAppStyle::GetAppStyleSetName(), "ClassIcon.Texture2D");
-							const FToolMenuExecuteAction UIAction = FToolMenuExecuteAction::CreateStatic(&ExecuteJPEGTextureSource);
-							InSection.AddMenuEntry("Texture2D_JPEGSource", Label, ToolTip, Icon, UIAction);
+							const TAttribute<FText> Label = LOCTEXT("Texture_SourceActions", "Texture Source Actions");
+							const TAttribute<FText> ToolTip = LOCTEXT("Texture_SourceActionsTooltip", "Actions that modify the texture source");
+							InSection.AddSubMenu("Texture_SourceActions", 
+												 Label, 
+												 ToolTip, 
+												 FNewToolMenuChoice { FOnGetContent::CreateStatic(&MakeTextureSourceActionMenu, Context) },
+												 false,
+												 FSlateIcon(FAppStyle::GetAppStyleSetName(), "ClassIcon.Texture2D"));
 						}
 					}
 				}));
