[Diff] Fix nullptr crash

#tests Diffed a change in main in 32.40 branch (specifically cl 39676295)

diff --git a/Engine/Plugins/Editor/EngineAssetDefinitions/Source/Private/AssetDefinition_DataAsset.cpp b/Engine/Plugins/Editor/EngineAssetDefinitions/Source/Private/AssetDefinition_DataAsset.cpp
--- a/Engine/Plugins/Editor/EngineAssetDefinitions/Source/Private/AssetDefinition_DataAsset.cpp
+++ b/Engine/Plugins/Editor/EngineAssetDefinitions/Source/Private/AssetDefinition_DataAsset.cpp
@@ -181,7 +181,7 @@ EAssetCommandResult UAssetDefinition_DataAsset::PerformAssetDiff(const FAssetDif
 	
 	const TSharedRef<SDetailsDiff> DetailsDiff = SDetailsDiff::CreateDiffWindow(DiffArgs.OldAsset, DiffArgs.NewAsset, DiffArgs.OldRevision, DiffArgs.NewRevision, UDataAsset::StaticClass());
 	// allow users to edit NewAsset if it's a local asset
-	if (!FPackageName::IsTempPackage(DiffArgs.NewAsset->GetPackage()->GetName()))
+	if (DiffArgs.NewAsset != nullptr && !FPackageName::IsTempPackage(DiffArgs.NewAsset->GetPackage()->GetName()))
 	{
 		DetailsDiff->SetOutputObject(DiffArgs.NewAsset);
 	}
