[AssetDefinition] Add bone icon of SkeletalMesh to the new Thumbnail status system

diff --git a/Engine/Plugins/Editor/EngineAssetDefinitions/Source/Private/Animation/AssetDefinition_SkeletalMesh.cpp b/Engine/Plugins/Editor/EngineAssetDefinitions/Source/Private/Animation/AssetDefinition_SkeletalMesh.cpp
--- a/Engine/Plugins/Editor/EngineAssetDefinitions/Source/Private/Animation/AssetDefinition_SkeletalMesh.cpp
+++ b/Engine/Plugins/Editor/EngineAssetDefinitions/Source/Private/Animation/AssetDefinition_SkeletalMesh.cpp
@@ -3,10 +3,12 @@
 #include "Animation/AssetDefinition_SkeletalMesh.h"
 
 #include "AnimationEditorUtils.h"
+#include "AssetDefinitionAssetInfo.h"
 #include "AssetNotifications.h"
 #include "ContentBrowserMenuContexts.h"
 #include "EditorReimportHandler.h"
 #include "FbxMeshUtils.h"
+#include "IAssetStatusInfoProvider.h"
 #include "ISkeletalMeshEditorModule.h"
 #include "Animation/SkeletalMeshUtilities.h"
 #include "AssetRegistry/AssetRegistryModule.h"
@@ -75,6 +77,18 @@ EAssetCommandResult UAssetDefinition_SkeletalMesh::OpenAssets(const FAssetOpenAr
 	return EAssetCommandResult::Handled;
 }
 
+void UAssetDefinition_SkeletalMesh::GetAssetStatusInfo(const TSharedPtr<IAssetStatusInfoProvider>& InAssetStatusInfoProvider, TArray<FAssetDisplayInfo>& OutStatusInfo) const
+{
+	Super::GetAssetStatusInfo(InAssetStatusInfoProvider, OutStatusInfo);
+
+	FAssetDisplayInfo SkinningStatus;
+	SkinningStatus.StatusIcon = FAppStyle::GetBrush("ClassThumbnailOverlays.SkeletalMesh_NeedSkinning");
+	SkinningStatus.Priority = FAssetStatusPriority(EStatusSeverity::Warning, 1);
+	SkinningStatus.StatusDescription = LOCTEXT("AssetDefintion_SkeletalMesh_NeedSkinningStatus_ToolTip", "Asset geometry was imported, the skinning needs to be validated");
+	SkinningStatus.IsVisible = MakeAttributeUObject(this, &UAssetDefinition_SkeletalMesh::GetThumbnailSkinningStatusVisibility, InAssetStatusInfoProvider);
+	OutStatusInfo.Add(SkinningStatus);
+}
+
 UThumbnailInfo* UAssetDefinition_SkeletalMesh::LoadThumbnailInfo(const FAssetData& InAssetData) const
 {
 	return UE::Editor::FindOrCreateThumbnailInfo(InAssetData.GetAsset(), USceneThumbnailInfo::StaticClass());
@@ -112,6 +126,16 @@ EVisibility UAssetDefinition_SkeletalMesh::GetThumbnailSkinningOverlayVisibility
 	return EVisibility::Collapsed;
 }
 
+EVisibility UAssetDefinition_SkeletalMesh::GetThumbnailSkinningStatusVisibility(const TSharedPtr<IAssetStatusInfoProvider> InAssetStatusInfoProvider) const
+{
+	const FAssetData AssetData = InAssetStatusInfoProvider->TryGetAssetData();
+	if (AssetData.IsValid())
+	{
+		return GetThumbnailSkinningOverlayVisibility(AssetData);
+	}
+	return EVisibility::Collapsed;
+}
+
 void UAssetDefinition_SkeletalMesh::OnAssetRemoved(const FAssetData& AssetData) const
 {
 	//Remove the object from the list before it get garbage collect
diff --git a/Engine/Plugins/Editor/EngineAssetDefinitions/Source/Private/Animation/AssetDefinition_SkeletalMesh.h b/Engine/Plugins/Editor/EngineAssetDefinitions/Source/Private/Animation/AssetDefinition_SkeletalMesh.h
--- a/Engine/Plugins/Editor/EngineAssetDefinitions/Source/Private/Animation/AssetDefinition_SkeletalMesh.h
+++ b/Engine/Plugins/Editor/EngineAssetDefinitions/Source/Private/Animation/AssetDefinition_SkeletalMesh.h
@@ -26,6 +26,7 @@ public:
 		static const auto Categories = { EAssetCategoryPaths::Basic, EAssetCategoryPaths::Animation };
 		return Categories;
 	}
+	virtual void GetAssetStatusInfo(const TSharedPtr<IAssetStatusInfoProvider>& InAssetStatusInfoProvider, TArray<FAssetDisplayInfo>& OutStatusInfo) const override;
 	virtual bool CanImport() const override { return true; }
 	virtual UThumbnailInfo* LoadThumbnailInfo(const FAssetData& InAssetData) const override;
 	virtual TSharedPtr<class SWidget> GetThumbnailOverlay(const FAssetData& AssetData) const override;
@@ -35,6 +36,7 @@ public:
 private:
 	void OnAssetRemoved(const FAssetData& AssetData) const;
 	EVisibility GetThumbnailSkinningOverlayVisibility(const FAssetData AssetData) const;
+	EVisibility GetThumbnailSkinningStatusVisibility(const TSharedPtr<IAssetStatusInfoProvider> InAssetStatusInfoProvider) const;
 
 	mutable TArray<FString> ThumbnailSkinningOverlayAssetNames;
 };
diff --git a/Engine/Source/Editor/AssetDefinition/Private/AssetStatusAssetDataInfoProvider.cpp b/Engine/Source/Editor/AssetDefinition/Private/AssetStatusAssetDataInfoProvider.cpp
--- a/Engine/Source/Editor/AssetDefinition/Private/AssetStatusAssetDataInfoProvider.cpp
+++ b/Engine/Source/Editor/AssetDefinition/Private/AssetStatusAssetDataInfoProvider.cpp
@@ -18,3 +18,8 @@ FString FAssetStatusAssetDataInfoProvider::TryGetFilename() const
 	FPackageName::TryConvertLongPackageNameToFilename(AssetData.PackageName.ToString(), OutFileName, *PackageExtension);
 	return OutFileName;
 }
+
+FAssetData FAssetStatusAssetDataInfoProvider::TryGetAssetData() const
+{
+	return AssetData.IsValid() ? AssetData : FAssetData();
+}
diff --git a/Engine/Source/Editor/AssetDefinition/Public/AssetStatusAssetDataInfoProvider.h b/Engine/Source/Editor/AssetDefinition/Public/AssetStatusAssetDataInfoProvider.h
--- a/Engine/Source/Editor/AssetDefinition/Public/AssetStatusAssetDataInfoProvider.h
+++ b/Engine/Source/Editor/AssetDefinition/Public/AssetStatusAssetDataInfoProvider.h
@@ -16,6 +16,8 @@ public:
 
 	virtual FString TryGetFilename() const override;
 
+	virtual FAssetData TryGetAssetData() const override;
+
 private:
 	FAssetData AssetData;
 };
diff --git a/Engine/Source/Editor/AssetDefinition/Public/IAssetStatusInfoProvider.h b/Engine/Source/Editor/AssetDefinition/Public/IAssetStatusInfoProvider.h
--- a/Engine/Source/Editor/AssetDefinition/Public/IAssetStatusInfoProvider.h
+++ b/Engine/Source/Editor/AssetDefinition/Public/IAssetStatusInfoProvider.h
@@ -2,6 +2,7 @@
 
 #pragma once
 
+#include "AssetRegistry/AssetData.h"
 #include "Containers/UnrealString.h"
 
 class UPackage;
@@ -16,4 +17,7 @@ public:
 
 	/** Try to get the filename, return empty otherwise */
 	virtual FString TryGetFilename() const = 0;
+
+	/** Try to get the AssetData, return an empty FAssetData otherwise */
+	virtual FAssetData TryGetAssetData() const = 0;
 };
