UE: Pose asset source update context menu option now uses the same settings as are used in PostLoad

diff --git a/Engine/Plugins/Editor/EngineAssetDefinitions/Source/Private/Animation/AssetDefinition_PoseAsset.cpp b/Engine/Plugins/Editor/EngineAssetDefinitions/Source/Private/Animation/AssetDefinition_PoseAsset.cpp
--- a/Engine/Plugins/Editor/EngineAssetDefinitions/Source/Private/Animation/AssetDefinition_PoseAsset.cpp
+++ b/Engine/Plugins/Editor/EngineAssetDefinitions/Source/Private/Animation/AssetDefinition_PoseAsset.cpp
@@ -27,7 +27,9 @@ namespace MenuExtension_PoseAsset
 			{
 				if (IAnimationDataModel* SourceAnimationDataModel = SourceAnimation->GetDataModel())
 				{
-					if (!PoseAsset->SourceAnimationRawDataGUID.IsValid() || PoseAsset->SourceAnimationRawDataGUID != SourceAnimationDataModel->GenerateGuid())
+					IAnimationDataModel::FGuidGenerationSettings Settings;
+					Settings.bIncludeTimingData = false;
+					if (!PoseAsset->SourceAnimationRawDataGUID.IsValid() || PoseAsset->SourceAnimationRawDataGUID != SourceAnimationDataModel->GenerateGuid(Settings))
 					{
 						if (PoseAsset->GetSkeleton()->IsCompatibleForEditor(SourceAnimation->GetSkeleton()))
 						{
