Prevent UDataAsset from being changed to types that are not allowed.

- When displaying the list of classes to change a UDataAsset make sure we are using the same set of filters as when creating a UDataAsset, found in UDataAssetFactory::ConfigureProperties.

diff --git a/Engine/Plugins/Editor/EngineAssetDefinitions/Source/Private/AssetDefinition_DataAsset.cpp b/Engine/Plugins/Editor/EngineAssetDefinitions/Source/Private/AssetDefinition_DataAsset.cpp
--- a/Engine/Plugins/Editor/EngineAssetDefinitions/Source/Private/AssetDefinition_DataAsset.cpp
+++ b/Engine/Plugins/Editor/EngineAssetDefinitions/Source/Private/AssetDefinition_DataAsset.cpp
@@ -25,6 +25,7 @@ namespace MenuExtension_DataAsset
 	public:
 		FNewNodeClassFilter(UClass* InBaseClass)
 			: BaseClass(InBaseClass)
+			, DisallowedClassFlags(CLASS_Abstract | CLASS_Deprecated | CLASS_NewerVersionExists | CLASS_HideDropDown)
 		{
 		}
 
@@ -32,18 +33,20 @@ namespace MenuExtension_DataAsset
 		{
 			if (InClass != nullptr)
 			{
-				return InClass->IsChildOf(BaseClass);
+				return InClass->IsChildOf(BaseClass) && !InClass->HasAnyClassFlags(DisallowedClassFlags);
 			}
 			return false;
 		}
 
 		virtual bool IsUnloadedClassAllowed(const FClassViewerInitializationOptions& InInitOptions, const TSharedRef< const IUnloadedBlueprintData > InUnloadedClassData, TSharedRef< FClassViewerFilterFuncs > InFilterFuncs) override
 		{
-			return InUnloadedClassData->IsChildOf(BaseClass);
+			return InUnloadedClassData->IsChildOf(BaseClass) && !InUnloadedClassData->HasAnyClassFlags(DisallowedClassFlags);
 		}
 
 	private:
 		UClass* BaseClass;
+
+		EClassFlags DisallowedClassFlags;
 	};
 
 	bool IsChangeDataAssetClassVisible(const FToolMenuContext& MenuContext)
@@ -129,7 +132,7 @@ namespace MenuExtension_DataAsset
 		{
 			FToolMenuOwnerScoped OwnerScoped(UE_MODULE_NAME);
 			UToolMenu* Menu = UE::ContentBrowser::ExtendToolMenu_AssetContextMenu(UDataAsset::StaticClass());
-	        
+
 			FToolMenuSection& Section = Menu->FindOrAddSection("GetAssetActions");
 			Section.AddDynamicEntry(NAME_None, FNewToolMenuSectionDelegate::CreateStatic([](FToolMenuSection& InSection)
 			{
