[EngineAssetDefinition] Fix SkeletalMesh regression for the RedBone icon where its not visible.

diff --git a/Engine/Plugins/Editor/EngineAssetDefinitions/Source/Private/Animation/AssetDefinition_SkeletalMesh.cpp b/Engine/Plugins/Editor/EngineAssetDefinitions/Source/Private/Animation/AssetDefinition_SkeletalMesh.cpp
--- a/Engine/Plugins/Editor/EngineAssetDefinitions/Source/Private/Animation/AssetDefinition_SkeletalMesh.cpp
+++ b/Engine/Plugins/Editor/EngineAssetDefinitions/Source/Private/Animation/AssetDefinition_SkeletalMesh.cpp
@@ -83,6 +83,12 @@ void UAssetDefinition_SkeletalMesh::GetAssetStatusInfo(const TSharedPtr<IAssetSt
 {
 	Super::GetAssetStatusInfo(InAssetStatusInfoProvider, OutStatusInfo);
 
+	const FAssetData AssetData = InAssetStatusInfoProvider->TryGetAssetData();
+	if (AssetData.IsValid())
+	{
+		ThumbnailSkinningOverlayAssetNames.AddUnique(AssetData.GetFullName());
+	}
+
 	FAssetDisplayInfo SkinningStatus;
 	SkinningStatus.StatusIcon = FAppStyle::GetBrush("ClassThumbnailOverlays.SkeletalMesh_NeedSkinning");
 	SkinningStatus.Priority = FAssetStatusPriority(EStatusSeverity::Warning, 1);
diff --git a/Engine/Source/Editor/UnrealEd/Private/AssetThumbnail.cpp b/Engine/Source/Editor/UnrealEd/Private/AssetThumbnail.cpp
--- a/Engine/Source/Editor/UnrealEd/Private/AssetThumbnail.cpp
+++ b/Engine/Source/Editor/UnrealEd/Private/AssetThumbnail.cpp
@@ -1123,9 +1123,13 @@ private:
 		if (OverlayInfo.IsValidIndex(StatusIndex))
 		{
 			const FAssetDisplayInfo& AssetStatus = OverlayInfo[StatusIndex];
-			if (AssetStatus.IsVisible.IsSet() && AssetStatus.IsVisible.Get() == EVisibility::Visible)
+			if (AssetStatus.IsVisible.IsSet())
 			{
-				return GetStatusVisibilityBasedOnGeometry(StatusIndex);
+				const EVisibility StatusVisibility = AssetStatus.IsVisible.Get();
+				if (StatusVisibility != EVisibility::Hidden && StatusVisibility != EVisibility::Collapsed)
+				{
+					return GetStatusVisibilityBasedOnGeometry(StatusIndex);
+				}
 			}
 		}
 		return EVisibility::Collapsed;
