Add a modular feature that physics assets will query when opening allowing for a decoupled bypass feature where a plugin can implement a new editor paradigm for physics assets.

diff --git a/Engine/Plugins/Editor/EngineAssetDefinitions/Source/Private/AssetDefinition_PhysicsAsset.cpp b/Engine/Plugins/Editor/EngineAssetDefinitions/Source/Private/AssetDefinition_PhysicsAsset.cpp
--- a/Engine/Plugins/Editor/EngineAssetDefinitions/Source/Private/AssetDefinition_PhysicsAsset.cpp
+++ b/Engine/Plugins/Editor/EngineAssetDefinitions/Source/Private/AssetDefinition_PhysicsAsset.cpp
@@ -4,6 +4,8 @@
 #include "Modules/ModuleManager.h"
 #include "PhysicsAssetEditorModule.h"
 #include "ThumbnailRendering/SceneThumbnailInfo.h"
+#include "IPhysicsAssetEditor.h"
+#include "Features/IModularFeatures.h"
 
 #include UE_INLINE_GENERATED_CPP_BY_NAME(AssetDefinition_PhysicsAsset)
 
@@ -16,9 +18,21 @@ UThumbnailInfo* UAssetDefinition_PhysicsAsset::LoadThumbnailInfo(const FAssetDat
 
 EAssetCommandResult UAssetDefinition_PhysicsAsset::OpenAssets(const FAssetOpenArgs& OpenArgs) const
 {
+	TArray<IPhysicsAssetEditorOverride*> ExternalEditors = IModularFeatures::Get().GetModularFeatureImplementations<IPhysicsAssetEditorOverride>(IPhysicsAssetEditorOverride::ModularFeatureName);
+
 	for (UPhysicsAsset* PhysicsAsset : OpenArgs.LoadObjects<UPhysicsAsset>())
 	{
-		IPhysicsAssetEditorModule* PhysicsAssetEditorModule = &FModuleManager::LoadModuleChecked<IPhysicsAssetEditorModule>( "PhysicsAssetEditor" );
+		bool bHandledExternally = false;
+
+		for(IPhysicsAssetEditorOverride* OverrideEditor : ExternalEditors)
+		{
+			if(OverrideEditor->OpenAsset(PhysicsAsset))
+			{
+				return EAssetCommandResult::Handled;
+			}
+		}
+
+		IPhysicsAssetEditorModule* PhysicsAssetEditorModule = &FModuleManager::LoadModuleChecked<IPhysicsAssetEditorModule>("PhysicsAssetEditor");
 		PhysicsAssetEditorModule->CreatePhysicsAssetEditor(OpenArgs.GetToolkitMode(), OpenArgs.ToolkitHost, PhysicsAsset);
 	}
 
diff --git a/Engine/Source/Editor/PhysicsAssetEditor/Public/IPhysicsAssetEditor.h b/Engine/Source/Editor/PhysicsAssetEditor/Public/IPhysicsAssetEditor.h
--- a/Engine/Source/Editor/PhysicsAssetEditor/Public/IPhysicsAssetEditor.h
+++ b/Engine/Source/Editor/PhysicsAssetEditor/Public/IPhysicsAssetEditor.h
@@ -5,7 +5,10 @@
 #include "CoreMinimal.h"
 #include "PersonaAssetEditorToolkit.h"
 #include "IHasPersonaToolkit.h"
+#include "Features/IModularFeature.h"
+#include "AssetDefinition.h"
 
+class UPhysicsAsset;
 
 /*-----------------------------------------------------------------------------
    IPhysicsAssetEditor
@@ -15,3 +18,11 @@ class IPhysicsAssetEditor : public FPersonaAssetEditorToolkit, public IHasPerson
 {
 
 };
+
+class IPhysicsAssetEditorOverride : public IModularFeature
+{
+public:
+	static const inline FName ModularFeatureName = "PhysicsAssetEditorOverride";
+
+	virtual bool OpenAsset(UPhysicsAsset*) = 0;
+};
\ No newline at end of file
