Ensure override editors can open multiple assets

diff --git a/Engine/Plugins/Editor/EngineAssetDefinitions/Source/Private/AssetDefinition_PhysicsAsset.cpp b/Engine/Plugins/Editor/EngineAssetDefinitions/Source/Private/AssetDefinition_PhysicsAsset.cpp
--- a/Engine/Plugins/Editor/EngineAssetDefinitions/Source/Private/AssetDefinition_PhysicsAsset.cpp
+++ b/Engine/Plugins/Editor/EngineAssetDefinitions/Source/Private/AssetDefinition_PhysicsAsset.cpp
@@ -28,10 +28,15 @@ EAssetCommandResult UAssetDefinition_PhysicsAsset::OpenAssets(const FAssetOpenAr
 		{
 			if(OverrideEditor->OpenAsset(PhysicsAsset))
 			{
-				return EAssetCommandResult::Handled;
+				bHandledExternally = true;
 			}
 		}
 
+		if(bHandledExternally)
+		{
+			continue;
+		}
+
 		IPhysicsAssetEditorModule* PhysicsAssetEditorModule = &FModuleManager::LoadModuleChecked<IPhysicsAssetEditorModule>("PhysicsAssetEditor");
 		PhysicsAssetEditorModule->CreatePhysicsAssetEditor(OpenArgs.GetToolkitMode(), OpenArgs.ToolkitHost, PhysicsAsset);
 	}
