Static Mesh LOD:  add clamping by MAX_STATIC_MESH_LODS in a few places where static meshes are modified, where it wasn't present.  Goal is to close off all code pathways that could allow content to be created that exceeds the LOD count, leading to runtime asserts.

#rnx

diff --git a/Engine/Plugins/Editor/EngineAssetDefinitions/Source/Private/AssetDefinition_StaticMesh.cpp b/Engine/Plugins/Editor/EngineAssetDefinitions/Source/Private/AssetDefinition_StaticMesh.cpp
--- a/Engine/Plugins/Editor/EngineAssetDefinitions/Source/Private/AssetDefinition_StaticMesh.cpp
+++ b/Engine/Plugins/Editor/EngineAssetDefinitions/Source/Private/AssetDefinition_StaticMesh.cpp
@@ -9,6 +9,7 @@
 #include "FbxMeshUtils.h"
 #include "Misc/MessageDialog.h"
 #include "StaticMeshCompiler.h"
+#include "StaticMeshResources.h"
 #include "ThumbnailRendering/SceneThumbnailInfo.h"
 #include "ToolMenu.h"
 #include "ToolMenus.h"
@@ -363,8 +364,8 @@ namespace MenuExtension_StaticMesh
 		};
 
 		TArray<FLODSettings> LODSettings;
-		LODSettings.AddZeroed(LODCopyMesh->GetNumSourceModels());
-		for (int32 i = 0; i < LODCopyMesh->GetNumSourceModels(); i++)
+		LODSettings.AddZeroed(FMath::Min(LODCopyMesh->GetNumSourceModels(), MAX_STATIC_MESH_LODS));
+		for (int32 i = 0; i < LODSettings.Num(); i++)
 		{
 			LODSettings[i].ReductionSettings = LODCopyMesh->GetSourceModel(i).ReductionSettings;
 			LODSettings[i].ScreenSize = LODCopyMesh->GetSourceModel(i).ScreenSize.Default;
diff --git a/Engine/Plugins/Runtime/MeshModelingToolset/Source/ModelingComponentsEditorOnly/Private/AssetUtils/CreateStaticMeshUtil.cpp b/Engine/Plugins/Runtime/MeshModelingToolset/Source/ModelingComponentsEditorOnly/Private/AssetUtils/CreateStaticMeshUtil.cpp
--- a/Engine/Plugins/Runtime/MeshModelingToolset/Source/ModelingComponentsEditorOnly/Private/AssetUtils/CreateStaticMeshUtil.cpp
+++ b/Engine/Plugins/Runtime/MeshModelingToolset/Source/ModelingComponentsEditorOnly/Private/AssetUtils/CreateStaticMeshUtil.cpp
@@ -18,6 +18,7 @@
 #include "Materials/MaterialInterface.h"
 #include "AssetUtils/StaticMeshMaterialUtil.h"
 #include "StaticMeshAttributes.h"
+#include "StaticMeshResources.h"
 
 using namespace UE::AssetUtils;
 
@@ -66,7 +67,7 @@ UE::AssetUtils::ECreateStaticMeshResult UE::AssetUtils::CreateStaticMeshAsset(
 	}
 
 	// initialize the MeshDescription SourceModel LODs
-	int32 UseNumSourceModels = FMath::Max(1, Options.NumSourceModels);
+	int32 UseNumSourceModels = FMath::Clamp(Options.NumSourceModels, 1, MAX_STATIC_MESH_LODS);
 	NewStaticMesh->SetNumSourceModels(UseNumSourceModels);
 	for (int32 k = 0; k < UseNumSourceModels; ++k)
 	{
diff --git a/Engine/Source/Runtime/Engine/Private/StaticMesh.cpp b/Engine/Source/Runtime/Engine/Private/StaticMesh.cpp
--- a/Engine/Source/Runtime/Engine/Private/StaticMesh.cpp
+++ b/Engine/Source/Runtime/Engine/Private/StaticMesh.cpp
@@ -8743,7 +8743,7 @@ bool UStaticMesh::BuildFromMeshDescriptions(const TArray<const FMeshDescription*
 {
 	TRACE_CPUPROFILER_EVENT_SCOPE(UStaticMesh::BuildFromMeshDescriptions);
 
-	const int32 NewNumLODs = MeshDescriptions.Num();
+	const int32 NewNumLODs = FMath::Min(MeshDescriptions.Num(), MAX_STATIC_MESH_LODS);
 
 	// Set up
 	NeverStream = true;
diff --git a/Engine/Source/Runtime/Engine/Private/Streaming/StreamingTexture.h b/Engine/Source/Runtime/Engine/Private/Streaming/StreamingTexture.h
--- a/Engine/Source/Runtime/Engine/Private/Streaming/StreamingTexture.h
+++ b/Engine/Source/Runtime/Engine/Private/Streaming/StreamingTexture.h
@@ -9,6 +9,7 @@ StreamingTexture.h: Definitions of classes used for texture streaming.
 #include "CoreMinimal.h"
 #include "Engine/StreamableRenderAsset.h"
 #include "Engine/TextureDefines.h"
+#include "StaticMeshResources.h"
 
 struct FRenderAssetStreamingManager;
 struct FRenderAssetStreamingSettings;
@@ -218,6 +219,8 @@ struct FStreamingRenderAsset
 			float LODScreenSizes[FStreamableRenderResourceState::MAX_LOD_COUNT / 2];
 		};
 	};
+	// Make sure static mesh LODs will fit in the above struct
+	static_assert(FStreamableRenderResourceState::MAX_LOD_COUNT / 2 >= MAX_STATIC_MESH_LODS);
 
 	/** (2) How much game time has elapsed since the texture was bound for rendering. Based on FApp::GetCurrentTime(). */
 	float			LastRenderTime;
