UE: The "Update Source Animation" context menu option on pose assets now only update the source if it hasn't been intentionally broken (i.e. if the user manually adds poses or similar divergence from the source animation). The guid will be invalid in these cases. If you need to revert to the source you can still do that in persona normally. This menu option is moslty for fixing the "source animation is out of date" warning which does not occur when the link to the source was intentionally broken (i.e. invalid guid)

diff --git a/Engine/Plugins/Editor/EngineAssetDefinitions/Source/Private/Animation/AssetDefinition_PoseAsset.cpp b/Engine/Plugins/Editor/EngineAssetDefinitions/Source/Private/Animation/AssetDefinition_PoseAsset.cpp
--- a/Engine/Plugins/Editor/EngineAssetDefinitions/Source/Private/Animation/AssetDefinition_PoseAsset.cpp
+++ b/Engine/Plugins/Editor/EngineAssetDefinitions/Source/Private/Animation/AssetDefinition_PoseAsset.cpp
@@ -29,7 +29,7 @@ namespace MenuExtension_PoseAsset
 			{
 				if (IAnimationDataModel* SourceAnimationDataModel = SourceAnimation->GetDataModel())
 				{
-					if (!PoseAsset->SourceAnimationRawDataGUID.IsValid() || PoseAsset->SourceAnimationRawDataGUID != PoseAsset->GetSourceAnimationGuid())
+					if (PoseAsset->SourceAnimationRawDataGUID.IsValid() && PoseAsset->SourceAnimationRawDataGUID != PoseAsset->GetSourceAnimationGuid())
 					{
 						if (PoseAsset->GetSkeleton()->IsCompatibleForEditor(SourceAnimation->GetSkeleton()))
 						{
