EditorUsability : AssetDefinitions
- Added Material Instance to custom filter in content browser to search for these asset type

#Synced-CL 51208877

diff --git a/Engine/Plugins/Editor/EngineAssetDefinitions/Source/Private/Material/AssetDefinition_MaterialInstanceConstant.h b/Engine/Plugins/Editor/EngineAssetDefinitions/Source/Private/Material/AssetDefinition_MaterialInstanceConstant.h
--- a/Engine/Plugins/Editor/EngineAssetDefinitions/Source/Private/Material/AssetDefinition_MaterialInstanceConstant.h
+++ b/Engine/Plugins/Editor/EngineAssetDefinitions/Source/Private/Material/AssetDefinition_MaterialInstanceConstant.h
@@ -24,7 +24,7 @@ public:
 	virtual TSoftClassPtr<UObject> GetAssetClass() const override { return UMaterialInstanceConstant::StaticClass(); }
 	virtual TConstArrayView<FAssetCategoryPath> GetAssetCategories() const override
 	{
-		static const auto Categories = { EAssetCategoryPaths::Material };
+		static const auto Categories = { EAssetCategoryPaths::Material, EAssetCategoryPaths::Basic };
 		return Categories;
 	}
 	virtual EAssetCommandResult OpenAssets(const FAssetOpenArgs& OpenArgs) const override;
