[Subsurface Profile] The preset of default or when the profile has been editted is changed from `Select Preset...` to `Custom`.

#Synced-CL 51211449

diff --git a/Engine/Plugins/Editor/EngineAssetDefinitions/Source/Private/AssetDefinition_SubsurfaceProfile.cpp b/Engine/Plugins/Editor/EngineAssetDefinitions/Source/Private/AssetDefinition_SubsurfaceProfile.cpp
--- a/Engine/Plugins/Editor/EngineAssetDefinitions/Source/Private/AssetDefinition_SubsurfaceProfile.cpp
+++ b/Engine/Plugins/Editor/EngineAssetDefinitions/Source/Private/AssetDefinition_SubsurfaceProfile.cpp
@@ -655,7 +655,7 @@ private:
 			}
 		}
 
-		return LOCTEXT("SelectPreset", "Select Preset...");
+		return LOCTEXT("SelectPreset", "Custom");
 	}
 
 	void ApplyProfilePreset(int32 PresetIndex)
