[ContentBrowser] New Add Menu Material Menu

#Synced-CL 52417896

diff --git a/Engine/Plugins/Editor/EngineAssetDefinitions/Source/Private/AssetDefinition_NeuralProfile.h b/Engine/Plugins/Editor/EngineAssetDefinitions/Source/Private/AssetDefinition_NeuralProfile.h
--- a/Engine/Plugins/Editor/EngineAssetDefinitions/Source/Private/AssetDefinition_NeuralProfile.h
+++ b/Engine/Plugins/Editor/EngineAssetDefinitions/Source/Private/AssetDefinition_NeuralProfile.h
@@ -19,7 +19,7 @@ public:
 	virtual TSoftClassPtr<UObject> GetAssetClass() const override { return UNeuralProfile::StaticClass(); }
 	virtual TConstArrayView<FAssetCategoryPath> GetAssetCategories() const override
     {
-    	static const auto Categories = { EAssetCategoryPaths::Material / NSLOCTEXT("Material", "MaterialAssetSubMenu_Profiles", "Profiles") };
+    	static const auto Categories = { FAssetCategoryPath(EAssetCategoryPaths::Material, NSLOCTEXT("Material", "MaterialAssetSubMenu_Profiles", "Profiles"), ECategoryMenuType::Section) };
     	return Categories;
     }
 	// UAssetDefinition End
diff --git a/Engine/Plugins/Editor/EngineAssetDefinitions/Source/Private/AssetDefinition_SpecularProfile.h b/Engine/Plugins/Editor/EngineAssetDefinitions/Source/Private/AssetDefinition_SpecularProfile.h
--- a/Engine/Plugins/Editor/EngineAssetDefinitions/Source/Private/AssetDefinition_SpecularProfile.h
+++ b/Engine/Plugins/Editor/EngineAssetDefinitions/Source/Private/AssetDefinition_SpecularProfile.h
@@ -19,7 +19,7 @@ public:
 	virtual TSoftClassPtr<UObject> GetAssetClass() const override { return USpecularProfile::StaticClass(); }
 	virtual TConstArrayView<FAssetCategoryPath> GetAssetCategories() const override
     {
-    	static const auto Categories = { EAssetCategoryPaths::Material / NSLOCTEXT("Material", "MaterialAssetSubMenu_Profiles", "Profiles") };
+    	static const auto Categories = { FAssetCategoryPath(EAssetCategoryPaths::Material, NSLOCTEXT("Material", "MaterialAssetSubMenu_Profiles", "Profiles"), ECategoryMenuType::Section) };
     	return Categories;
     }
 	// UAssetDefinition End
diff --git a/Engine/Plugins/Editor/EngineAssetDefinitions/Source/Private/AssetDefinition_SubsurfaceProfile.h b/Engine/Plugins/Editor/EngineAssetDefinitions/Source/Private/AssetDefinition_SubsurfaceProfile.h
--- a/Engine/Plugins/Editor/EngineAssetDefinitions/Source/Private/AssetDefinition_SubsurfaceProfile.h
+++ b/Engine/Plugins/Editor/EngineAssetDefinitions/Source/Private/AssetDefinition_SubsurfaceProfile.h
@@ -19,7 +19,7 @@ public:
 	virtual TSoftClassPtr<UObject> GetAssetClass() const override { return USubsurfaceProfile::StaticClass(); }
 	virtual TConstArrayView<FAssetCategoryPath> GetAssetCategories() const override
     {
-    	static const auto Categories = { EAssetCategoryPaths::Material  / NSLOCTEXT("Material", "MaterialAssetSubMenu_Profiles", "Profiles") };
+    	static const auto Categories = { FAssetCategoryPath(EAssetCategoryPaths::Material, NSLOCTEXT("Material", "MaterialAssetSubMenu_Profiles", "Profiles"), ECategoryMenuType::Section) };
     	return Categories;
     }
 
diff --git a/Engine/Plugins/Editor/EngineAssetDefinitions/Source/Private/AssetDefinition_ToonProfile.h b/Engine/Plugins/Editor/EngineAssetDefinitions/Source/Private/AssetDefinition_ToonProfile.h
--- a/Engine/Plugins/Editor/EngineAssetDefinitions/Source/Private/AssetDefinition_ToonProfile.h
+++ b/Engine/Plugins/Editor/EngineAssetDefinitions/Source/Private/AssetDefinition_ToonProfile.h
@@ -19,7 +19,7 @@ public:
 	virtual TSoftClassPtr<UObject> GetAssetClass() const override { return UToonProfile::StaticClass(); }
 	virtual TConstArrayView<FAssetCategoryPath> GetAssetCategories() const override
     {
-    	static const auto Categories = { EAssetCategoryPaths::Material / NSLOCTEXT("Material", "MaterialAssetSubMenu_Profiles", "Profiles") };
+    	static const auto Categories = { FAssetCategoryPath(EAssetCategoryPaths::Material, NSLOCTEXT("Material", "MaterialAssetSubMenu_Profiles", "Profiles"), ECategoryMenuType::Section) };
     	return Categories;
     }
 	// UAssetDefinition End
diff --git a/Engine/Plugins/Editor/EngineAssetDefinitions/Source/Private/Material/AssetDefinition_MaterialFunction.h b/Engine/Plugins/Editor/EngineAssetDefinitions/Source/Private/Material/AssetDefinition_MaterialFunction.h
--- a/Engine/Plugins/Editor/EngineAssetDefinitions/Source/Private/Material/AssetDefinition_MaterialFunction.h
+++ b/Engine/Plugins/Editor/EngineAssetDefinitions/Source/Private/Material/AssetDefinition_MaterialFunction.h
@@ -20,7 +20,7 @@ public:
 	virtual TSoftClassPtr<UObject> GetAssetClass() const override { return UMaterialFunction::StaticClass(); }
 	virtual TConstArrayView<FAssetCategoryPath> GetAssetCategories() const override
 	{
-		static const auto Categories = { EAssetCategoryPaths::Material / NSLOCTEXT("Material", "MaterialAssetSubMenu_Advanced", "Advanced") };
+		static const auto Categories = { FAssetCategoryPath(EAssetCategoryPaths::Material, NSLOCTEXT("Material", "MaterialAssetSubMenu_Advanced", "Advanced"), ECategoryMenuType::Section) };
 		return Categories;
 	}
 	virtual UThumbnailInfo* LoadThumbnailInfo(const FAssetData& InAsset) const override;
@@ -40,7 +40,7 @@ public:
 	virtual TSoftClassPtr<UObject> GetAssetClass() const override { return UMaterialFunctionMaterialLayer::StaticClass(); }
 	virtual TConstArrayView<FAssetCategoryPath> GetAssetCategories() const override
 	{
-		static const auto Categories = { EAssetCategoryPaths::Material / NSLOCTEXT("Material", "MaterialAssetSubMenu_Layers", "Layers") };
+		static const auto Categories = { FAssetCategoryPath(EAssetCategoryPaths::Material, NSLOCTEXT("Material", "MaterialAssetSubMenu_Layers", "Layers"), ECategoryMenuType::Section) };
 		return Categories;
 	}
 	// UAssetDefinition End
@@ -57,7 +57,7 @@ public:
 	virtual TSoftClassPtr<UObject> GetAssetClass() const override { return UMaterialFunctionMaterialLayerBlend::StaticClass(); }
 	virtual TConstArrayView<FAssetCategoryPath> GetAssetCategories() const override
 	{
-		static const auto Categories = { EAssetCategoryPaths::Material / NSLOCTEXT("Material", "MaterialAssetSubMenu_Layers", "Layers") };
+		static const auto Categories = { FAssetCategoryPath(EAssetCategoryPaths::Material, NSLOCTEXT("Material", "MaterialAssetSubMenu_Layers", "Layers"), ECategoryMenuType::Section) };
 		return Categories;
 	}
 	// UAssetDefinition End
diff --git a/Engine/Plugins/Editor/EngineAssetDefinitions/Source/Private/Material/AssetDefinition_MaterialFunctionInstance.h b/Engine/Plugins/Editor/EngineAssetDefinitions/Source/Private/Material/AssetDefinition_MaterialFunctionInstance.h
--- a/Engine/Plugins/Editor/EngineAssetDefinitions/Source/Private/Material/AssetDefinition_MaterialFunctionInstance.h
+++ b/Engine/Plugins/Editor/EngineAssetDefinitions/Source/Private/Material/AssetDefinition_MaterialFunctionInstance.h
@@ -34,7 +34,7 @@ public:
 	virtual TSoftClassPtr<UObject> GetAssetClass() const override { return UMaterialFunctionMaterialLayerInstance::StaticClass(); }
 	virtual TConstArrayView<FAssetCategoryPath> GetAssetCategories() const override
 	{
-		static const auto Categories = { EAssetCategoryPaths::Material / NSLOCTEXT("Material", "MaterialAssetSubMenu_Layers", "Layers") };
+		static const auto Categories = { FAssetCategoryPath(EAssetCategoryPaths::Material, NSLOCTEXT("Material", "MaterialAssetSubMenu_Layers", "Layers"), ECategoryMenuType::Section) };
 		return Categories;
 	}
 	// UAssetDefinition End
@@ -51,7 +51,7 @@ public:
 	virtual TSoftClassPtr<UObject> GetAssetClass() const override { return UMaterialFunctionMaterialLayerBlendInstance::StaticClass(); }
 	virtual TConstArrayView<FAssetCategoryPath> GetAssetCategories() const override
 	{
-		static const auto Categories = { EAssetCategoryPaths::Material / NSLOCTEXT("Material", "MaterialAssetSubMenu_Layers", "Layers") };
+		static const auto Categories = { FAssetCategoryPath(EAssetCategoryPaths::Material, NSLOCTEXT("Material", "MaterialAssetSubMenu_Layers", "Layers"), ECategoryMenuType::Section) };
 		return Categories;
 	}
 	// UAssetDefinition End
diff --git a/Engine/Plugins/Editor/EngineAssetDefinitions/Source/Private/Material/AssetDefinition_MaterialInstanceConstant.h b/Engine/Plugins/Editor/EngineAssetDefinitions/Source/Private/Material/AssetDefinition_MaterialInstanceConstant.h
--- a/Engine/Plugins/Editor/EngineAssetDefinitions/Source/Private/Material/AssetDefinition_MaterialInstanceConstant.h
+++ b/Engine/Plugins/Editor/EngineAssetDefinitions/Source/Private/Material/AssetDefinition_MaterialInstanceConstant.h
@@ -24,7 +24,7 @@ public:
 	virtual TSoftClassPtr<UObject> GetAssetClass() const override { return UMaterialInstanceConstant::StaticClass(); }
 	virtual TConstArrayView<FAssetCategoryPath> GetAssetCategories() const override
 	{
-		static const auto Categories = { EAssetCategoryPaths::Material, EAssetCategoryPaths::Basic };
+		static const auto Categories = { FAssetCategoryPath(EAssetCategoryPaths::Material), FAssetCategoryPath(EAssetCategoryPaths::Basic) };
 		return Categories;
 	}
 	virtual EAssetCommandResult OpenAssets(const FAssetOpenArgs& OpenArgs) const override;
diff --git a/Engine/Plugins/Editor/EngineAssetDefinitions/Source/Private/Material/AssetDefinition_MaterialParameterCollection.h b/Engine/Plugins/Editor/EngineAssetDefinitions/Source/Private/Material/AssetDefinition_MaterialParameterCollection.h
--- a/Engine/Plugins/Editor/EngineAssetDefinitions/Source/Private/Material/AssetDefinition_MaterialParameterCollection.h
+++ b/Engine/Plugins/Editor/EngineAssetDefinitions/Source/Private/Material/AssetDefinition_MaterialParameterCollection.h
@@ -19,7 +19,7 @@ public:
 	virtual TSoftClassPtr<UObject> GetAssetClass() const override { return UMaterialParameterCollection::StaticClass(); }
 	virtual TConstArrayView<FAssetCategoryPath> GetAssetCategories() const override
 	{
-		static const auto Categories = { EAssetCategoryPaths::Material };
+		static const auto Categories = { FAssetCategoryPath(EAssetCategoryPaths::Material) };
 		return Categories;
 	}
 	// UAssetDefinition End
