[ContentBrowser] New Add Menu Blueprint Menu

#Synced-CL 52439714

diff --git a/Engine/Plugins/Editor/EngineAssetDefinitions/Source/Private/AssetDefinition_EditorUtilityBlueprint.cpp b/Engine/Plugins/Editor/EngineAssetDefinitions/Source/Private/AssetDefinition_EditorUtilityBlueprint.cpp
--- a/Engine/Plugins/Editor/EngineAssetDefinitions/Source/Private/AssetDefinition_EditorUtilityBlueprint.cpp
+++ b/Engine/Plugins/Editor/EngineAssetDefinitions/Source/Private/AssetDefinition_EditorUtilityBlueprint.cpp
@@ -35,14 +35,8 @@ TSoftClassPtr<> UAssetDefinition_EditorUtilityBlueprint::GetAssetClass() const
 
 TConstArrayView<FAssetCategoryPath> UAssetDefinition_EditorUtilityBlueprint::GetAssetCategories() const
 {
-	IBlutilityModule* BlutilityModule = FModuleManager::GetModulePtr<IBlutilityModule>("Blutility");
-	if (BlutilityModule)
-	{
-		return BlutilityModule->GetAssetCategories();
-	}
-
-	static const FAssetCategoryPath FallbackCategory[] = { EAssetCategoryPaths::Misc };
-	return FallbackCategory;
+	static const auto Categories = { FAssetCategoryPath(EAssetCategoryPaths::Blueprint, LOCTEXT("EditorUtilityBlueprint_SubMenus", "Editor Utilities"), ECategoryMenuType::Section) };
+	return Categories;
 }
 
 namespace MenuExtension_EditorUtilityBlueprint
diff --git a/Engine/Plugins/Editor/EngineAssetDefinitions/Source/Private/AssetDefinition_EditorUtilityWidgetBlueprint.cpp b/Engine/Plugins/Editor/EngineAssetDefinitions/Source/Private/AssetDefinition_EditorUtilityWidgetBlueprint.cpp
--- a/Engine/Plugins/Editor/EngineAssetDefinitions/Source/Private/AssetDefinition_EditorUtilityWidgetBlueprint.cpp
+++ b/Engine/Plugins/Editor/EngineAssetDefinitions/Source/Private/AssetDefinition_EditorUtilityWidgetBlueprint.cpp
@@ -38,14 +38,8 @@ TSoftClassPtr<> UAssetDefinition_EditorUtilityWidgetBlueprint::GetAssetClass() c
 
 TConstArrayView<FAssetCategoryPath> UAssetDefinition_EditorUtilityWidgetBlueprint::GetAssetCategories() const
 {
-	IBlutilityModule* BlutilityModule = FModuleManager::GetModulePtr<IBlutilityModule>("Blutility");
-	if (BlutilityModule)
-	{
-		return BlutilityModule->GetAssetCategories();
-	}
-
-	static const TArray<FAssetCategoryPath, TFixedAllocator<1>> FallbackCategory = { EAssetCategoryPaths::Misc };
-	return FallbackCategory;
+	static const auto Categories = { FAssetCategoryPath(EAssetCategoryPaths::Blueprint, LOCTEXT("EditorUtilityWidgetBlueprint_SubMenus", "Editor Utilities"), ECategoryMenuType::Section) };
+	return Categories;
 }
 
 EAssetCommandResult UAssetDefinition_EditorUtilityWidgetBlueprint::OpenAssets(const FAssetOpenArgs& OpenArgs) const
diff --git a/Engine/Plugins/Editor/EngineAssetDefinitions/Source/Private/AssetDefinition_UserDefinedEnum.h b/Engine/Plugins/Editor/EngineAssetDefinitions/Source/Private/AssetDefinition_UserDefinedEnum.h
--- a/Engine/Plugins/Editor/EngineAssetDefinitions/Source/Private/AssetDefinition_UserDefinedEnum.h
+++ b/Engine/Plugins/Editor/EngineAssetDefinitions/Source/Private/AssetDefinition_UserDefinedEnum.h
@@ -19,7 +19,7 @@ public:
 	virtual TSoftClassPtr<UObject> GetAssetClass() const override { return UUserDefinedEnum::StaticClass(); }
 	virtual TConstArrayView<FAssetCategoryPath> GetAssetCategories() const override
 	{
-		static const auto Categories = { EAssetCategoryPaths::Blueprint };
+		static const auto Categories = { FAssetCategoryPath(EAssetCategoryPaths::Blueprint, NSLOCTEXT("AssetDefinition", "Enum", "Composite Variables"), ECategoryMenuType::Section) };
 		return Categories;
 	}
 	virtual FAssetSupportResponse CanLocalize(const FAssetData& InAsset) const override
diff --git a/Engine/Plugins/Editor/EngineAssetDefinitions/Source/Private/AssetDefinition_UserDefinedStruct.h b/Engine/Plugins/Editor/EngineAssetDefinitions/Source/Private/AssetDefinition_UserDefinedStruct.h
--- a/Engine/Plugins/Editor/EngineAssetDefinitions/Source/Private/AssetDefinition_UserDefinedStruct.h
+++ b/Engine/Plugins/Editor/EngineAssetDefinitions/Source/Private/AssetDefinition_UserDefinedStruct.h
@@ -19,7 +19,7 @@ public:
 	virtual TSoftClassPtr<UObject> GetAssetClass() const override { return UUserDefinedStruct::StaticClass(); }
 	virtual TConstArrayView<FAssetCategoryPath> GetAssetCategories() const override
 	{
-		static const auto Categories = { EAssetCategoryPaths::Blueprint };
+		static const auto Categories = { FAssetCategoryPath(EAssetCategoryPaths::Blueprint, NSLOCTEXT("AssetDefinition", "Struct", "Composite Variables"), ECategoryMenuType::Section) };
 		return Categories;
 	}
 	virtual FAssetSupportResponse CanLocalize(const FAssetData& InAsset) const override
