[ContentBrowser] New Add Menu Physics Menu

#Synced-CL 52619962

diff --git a/Engine/Plugins/ChaosClothAssetEditorCore/Source/ChaosClothAssetEditor/Private/ChaosClothAsset/AssetDefinition_ClothAsset.cpp b/Engine/Plugins/ChaosClothAssetEditorCore/Source/ChaosClothAssetEditor/Private/ChaosClothAsset/AssetDefinition_ClothAsset.cpp
--- a/Engine/Plugins/ChaosClothAssetEditorCore/Source/ChaosClothAssetEditor/Private/ChaosClothAsset/AssetDefinition_ClothAsset.cpp
+++ b/Engine/Plugins/ChaosClothAssetEditorCore/Source/ChaosClothAssetEditor/Private/ChaosClothAsset/AssetDefinition_ClothAsset.cpp
@@ -51,7 +51,10 @@ FLinearColor UAssetDefinition_ClothAsset::GetAssetColor() const
 
 TConstArrayView<FAssetCategoryPath> UAssetDefinition_ClothAsset::GetAssetCategories() const
 {
-	static const auto Categories = { EAssetCategoryPaths::Physics };
+	static const TArray<FAssetCategoryPath> Categories = 
+		{
+			FAssetCategoryPath(EAssetCategoryPaths::Physics, LOCTEXT("AssetDefinition_ClothAssetSubMenu", "Chaos Cloth"), ECategoryMenuType::Section)
+		};
 	return Categories;
 }
 
diff --git a/Engine/Plugins/Editor/EngineAssetDefinitions/Source/Private/AssetDefinition_PhysicalMaterial.h b/Engine/Plugins/Editor/EngineAssetDefinitions/Source/Private/AssetDefinition_PhysicalMaterial.h
--- a/Engine/Plugins/Editor/EngineAssetDefinitions/Source/Private/AssetDefinition_PhysicalMaterial.h
+++ b/Engine/Plugins/Editor/EngineAssetDefinitions/Source/Private/AssetDefinition_PhysicalMaterial.h
@@ -14,12 +14,27 @@ class UAssetDefinition_PhysicalMaterial : public UAssetDefinitionDefault
 
 public:
 	// UAssetDefinition Begin
-	virtual FText GetAssetDisplayName() const override { return NSLOCTEXT("AssetTypeActions", "AssetTypeActions_PhysicalMaterial", "Physical Material"); }
-	virtual FLinearColor GetAssetColor() const override { return FLinearColor(FColor(200,192,128)); }
-	virtual TSoftClassPtr<UObject> GetAssetClass() const override { return UPhysicalMaterial::StaticClass(); }
+	virtual FText GetAssetDisplayName() const override
+	{
+		return NSLOCTEXT("AssetTypeActions", "AssetTypeActions_PhysicalMaterial", "Physical Material");
+	}
+
+	virtual FLinearColor GetAssetColor() const override
+	{
+		return FLinearColor(FColor(200,192,128));
+	}
+
+	virtual TSoftClassPtr<UObject> GetAssetClass() const override
+	{
+		return UPhysicalMaterial::StaticClass();
+	}
+
 	virtual TConstArrayView<FAssetCategoryPath> GetAssetCategories() const override
 	{
-		static const auto Categories = { EAssetCategoryPaths::Physics };
+		static const TArray<FAssetCategoryPath> Categories = 
+			{
+				FAssetCategoryPath(EAssetCategoryPaths::Physics, NSLOCTEXT("AssetDefinition", "PhysicalMaterial_SubMenu", "Physical Material"), ECategoryMenuType::Section)
+			};
 		return Categories;
 	}
 	// UAssetDefinition End
diff --git a/Engine/Plugins/Editor/EngineAssetDefinitions/Source/Private/AssetDefinition_PhysicalMaterialMask.cpp b/Engine/Plugins/Editor/EngineAssetDefinitions/Source/Private/AssetDefinition_PhysicalMaterialMask.cpp
new file mode 100644
--- /dev/null
+++ b/Engine/Plugins/Editor/EngineAssetDefinitions/Source/Private/AssetDefinition_PhysicalMaterialMask.cpp
@@ -0,0 +1,119 @@
+// Copyright Epic Games, Inc. All Rights Reserved.
+
+#include "AssetDefinition_PhysicalMaterialMask.h"
+#include "ContentBrowserMenuContexts.h"
+#include "EditorFramework/AssetImportData.h"
+#include "PhysicalMaterialMaskImport.h"
+#include "ToolMenus.h"
+
+#define LOCTEXT_NAMESPACE "AssetDefinition_PhysicalMaterialMask"
+
+namespace MenuExtension::PhysicalMaterialMask
+{
+	void GetResolvedSourceFilePaths(const UPhysicalMaterialMask* InMaterialMask, TArray<FString>& OutSourceFilePaths)
+	{
+		if (!InMaterialMask)
+		{
+			OutSourceFilePaths.Emplace(FString());
+			return;
+		}
+
+		TArray<FString> ExtractedFilenames;
+		if (UAssetImportData* AssetImportData = InMaterialMask->AssetImportData)
+		{
+			if (AssetImportData->GetSourceFileCount() > 0)
+			{
+				if (!AssetImportData->GetSourceData().SourceFiles[0].RelativeFilename.IsEmpty())
+				{
+					InMaterialMask->AssetImportData->ExtractFilenames(ExtractedFilenames);
+				}
+			}
+		}
+
+		if (ExtractedFilenames.Num() > 0)
+		{
+			OutSourceFilePaths.Emplace(ExtractedFilenames[0]);
+		}
+		else
+		{
+			OutSourceFilePaths.Emplace(FString());
+		}
+	}
+
+	bool IsOneAssetSelected(const FToolMenuContext& InContext)
+	{
+		const UContentBrowserAssetContextMenuContext* CBContext = UContentBrowserAssetContextMenuContext::FindContextWithAssets(InContext);
+		if (CBContext == nullptr)
+		{
+			return false;
+		}
+
+		TArray<FAssetData> AssetData = CBContext->SelectedAssets;
+		return AssetData.Num() == 1;
+	}
+
+	void ExecuteImport(const FToolMenuContext& InContext)
+	{
+		const UContentBrowserAssetContextMenuContext* CBContext = UContentBrowserAssetContextMenuContext::FindContextWithAssets(InContext);
+		if (CBContext == nullptr)
+		{
+			return;
+		}
+
+		TArray<FAssetData> AssetData = CBContext->SelectedAssets;
+		if (AssetData.Num() != 1)
+		{
+			return;
+		}
+
+		TArray<FString> OutResolvedFilePaths;
+		UPhysicalMaterialMask* SelectedMask = CBContext->LoadFirstSelectedObject<UPhysicalMaterialMask>();
+		GetResolvedSourceFilePaths(SelectedMask, OutResolvedFilePaths);
+
+		if (SelectedMask && OutResolvedFilePaths.Num() == 1 && OutResolvedFilePaths[0].IsEmpty())
+		{
+			FPhysicalMaterialMaskImport::ImportMaskTexture(SelectedMask);
+		}
+	}
+
+	static FDelayedAutoRegisterHelper DelayedAutoRegister(EDelayedRegisterRunPhase::EndOfEngineInit, []
+	{
+		UToolMenus::RegisterStartupCallback(FSimpleMulticastDelegate::FDelegate::CreateLambda([]()
+		{
+			FToolMenuOwnerScoped OwnerScoped(UE_MODULE_NAME);
+
+			// Physical Material Mask Action Registration
+			{
+				UToolMenu* Menu = UE::ContentBrowser::ExtendToolMenu_AssetContextMenu(UPhysicalMaterialMask::StaticClass());
+
+				FToolMenuSection& PhysicalMaterialMaskSection = Menu->FindOrAddSection("GetAssetActions");
+				PhysicalMaterialMaskSection.AddDynamicEntry(NAME_None, FNewToolMenuSectionDelegate::CreateLambda([](FToolMenuSection& InSection)
+					{
+						if (const UContentBrowserAssetContextMenuContext* Context = UContentBrowserAssetContextMenuContext::FindContextWithAssets(InSection))
+						{
+							if (Context->SelectedAssets.Num() > 0)
+							{
+								const FAssetData& AssetData = Context->SelectedAssets[0];
+								if (AssetData.AssetClassPath == UPhysicalMaterialMask::StaticClass()->GetClassPathName())
+								{
+									// Import Mask Texture
+									{
+										const TAttribute<FText> Label = LOCTEXT("PhysicalMaterialMask_ImportMaskTextureLabel", "Import Mask Texture");
+										const TAttribute<FText> ToolTip = LOCTEXT("PhysicalMaterialMask_ImportMaskTextureTooltip", "Imports a texture to use as a physical material mask.");
+										const FSlateIcon Icon = FSlateIcon(FAppStyle::GetAppStyleSetName(), "ClassIcon.Material");
+
+										FToolUIAction UIAction;
+										UIAction.ExecuteAction = FToolMenuExecuteAction::CreateStatic(&ExecuteImport);
+										UIAction.IsActionVisibleDelegate = FToolMenuIsActionButtonVisible::CreateStatic(&IsOneAssetSelected);
+										InSection.AddMenuEntry("PhysicalMaterialMask_Import", Label, ToolTip, Icon, UIAction);
+									}
+								}
+							}
+						}
+					}));
+			}
+		}));
+	});
+}
+
+#undef LOCTEXT_NAMESPACE
diff --git a/Engine/Plugins/Editor/EngineAssetDefinitions/Source/Private/AssetDefinition_PhysicalMaterialMask.h b/Engine/Plugins/Editor/EngineAssetDefinitions/Source/Private/AssetDefinition_PhysicalMaterialMask.h
new file mode 100644
--- /dev/null
+++ b/Engine/Plugins/Editor/EngineAssetDefinitions/Source/Private/AssetDefinition_PhysicalMaterialMask.h
@@ -0,0 +1,44 @@
+// Copyright Epic Games, Inc. All Rights Reserved.
+
+#pragma once
+#include "AssetDefinitionDefault.h"
+#include "PhysicalMaterials/PhysicalMaterialMask.h"
+#include "AssetDefinition_PhysicalMaterialMask.generated.h"
+
+UCLASS()
+class UAssetDefinition_PhysicalMaterialMask : public UAssetDefinitionDefault
+{
+	GENERATED_BODY()
+
+public:
+	// UAssetDefinition Begin
+	virtual FText GetAssetDisplayName() const override
+	{
+		return NSLOCTEXT("AssetDefinition", "AssetDefinition_PhysicalMaterialMask", "Physical Material Mask");
+	}
+
+	virtual FLinearColor GetAssetColor() const override
+	{
+		return FColor(200,162,108);
+	}
+
+	virtual TSoftClassPtr<UObject> GetAssetClass() const override
+	{
+		return UPhysicalMaterialMask::StaticClass();
+	}
+
+	virtual bool CanImport() const override
+	{
+		return true;
+	}
+
+	virtual TConstArrayView<FAssetCategoryPath> GetAssetCategories() const override
+	{
+		static const TArray<FAssetCategoryPath> Categories = 
+			{
+				FAssetCategoryPath(EAssetCategoryPaths::Physics, NSLOCTEXT("AssetDefinition", "AssetDefinition_PhysicalMaterialMaskSubMenu", "Physical Material"), ECategoryMenuType::Section)
+			};
+		return Categories;
+	}
+	// UAssetDefinition End
+};
diff --git a/Engine/Plugins/Editor/EngineAssetDefinitions/Source/Private/AssetDefinition_PhysicsAsset.h b/Engine/Plugins/Editor/EngineAssetDefinitions/Source/Private/AssetDefinition_PhysicsAsset.h
--- a/Engine/Plugins/Editor/EngineAssetDefinitions/Source/Private/AssetDefinition_PhysicsAsset.h
+++ b/Engine/Plugins/Editor/EngineAssetDefinitions/Source/Private/AssetDefinition_PhysicsAsset.h
@@ -14,14 +14,30 @@ class UAssetDefinition_PhysicsAsset : public UAssetDefinitionDefault
 
 public:
 	// UAssetDefinition Begin
-	virtual FText GetAssetDisplayName() const override { return NSLOCTEXT("AssetTypeActions", "AssetTypeActions_PhysicsAsset", "Physics Asset"); }
-	virtual FLinearColor GetAssetColor() const override { return FLinearColor(FColor(255,192,128)); }
-	virtual TSoftClassPtr<UObject> GetAssetClass() const override { return UPhysicsAsset::StaticClass(); }
+	virtual FText GetAssetDisplayName() const override
+	{
+		return NSLOCTEXT("AssetTypeActions", "AssetTypeActions_PhysicsAsset", "Physics Asset");
+	}
+
+	virtual FLinearColor GetAssetColor() const override
+	{
+		return FLinearColor(FColor(255,192,128));
+	}
+
+	virtual TSoftClassPtr<UObject> GetAssetClass() const override
+	{
+		return UPhysicsAsset::StaticClass();
+	}
+
 	virtual TConstArrayView<FAssetCategoryPath> GetAssetCategories() const override
 	{
-		static const auto Categories = { EAssetCategoryPaths::Physics };
+		static const TArray<FAssetCategoryPath> Categories = 
+			{
+				FAssetCategoryPath(EAssetCategoryPaths::Physics, NSLOCTEXT("AssetDefinition", "PhysicalMaterial_SubMenu", "Physical Material"), ECategoryMenuType::Section)
+			};
 		return Categories;
 	}
+
 	virtual EAssetCommandResult OpenAssets(const FAssetOpenArgs& OpenArgs) const override;
 	virtual UThumbnailInfo* LoadThumbnailInfo(const FAssetData& InAsset) const override;
 	// UAssetDefinition End
diff --git a/Engine/Plugins/Experimental/ChaosCaching/Source/ChaosCachingEditor/ChaosCachingEditor.Build.cs b/Engine/Plugins/Experimental/ChaosCaching/Source/ChaosCachingEditor/ChaosCachingEditor.Build.cs
--- a/Engine/Plugins/Experimental/ChaosCaching/Source/ChaosCachingEditor/ChaosCachingEditor.Build.cs
+++ b/Engine/Plugins/Experimental/ChaosCaching/Source/ChaosCachingEditor/ChaosCachingEditor.Build.cs
@@ -37,6 +37,12 @@ namespace UnrealBuildTool.Rules
 					"TakeTrackRecorders",
 				});
 
+			PrivateDependencyModuleNames.AddRange(
+				new string[]
+				{
+					"AssetDefinition"
+				});
+
 			SetupModulePhysicsSupport(Target);
 			PrivateDefinitions.Add("CHAOS_INCLUDE_LEVEL_1=1");
 		}
diff --git a/Engine/Plugins/Experimental/ChaosCaching/Source/ChaosCachingEditor/Private/Chaos/AssetDefinition_ChaosCacheCollection.h b/Engine/Plugins/Experimental/ChaosCaching/Source/ChaosCachingEditor/Private/Chaos/AssetDefinition_ChaosCacheCollection.h
new file mode 100644
--- /dev/null
+++ b/Engine/Plugins/Experimental/ChaosCaching/Source/ChaosCachingEditor/Private/Chaos/AssetDefinition_ChaosCacheCollection.h
@@ -0,0 +1,43 @@
+// Copyright Epic Games, Inc. All Rights Reserved.
+
+#pragma once
+#include "AssetDefinitionDefault.h"
+#include "Chaos/CacheCollection.h"
+#include "AssetDefinition_ChaosCacheCollection.generated.h"
+
+UCLASS()
+class UAssetDefinition_ChaosCacheCollection : public UAssetDefinitionDefault
+{
+	GENERATED_BODY()
+public:
+	// UAssetDefinition Begin
+	virtual FText GetAssetDisplayName() const override
+	{
+		return NSLOCTEXT("AssetDefinition", "AssetDefinition_ChaosCacheCollection", "Chaos Cache Collection");
+	}
+
+	virtual FLinearColor GetAssetColor() const override
+	{
+		return FColor(255, 127, 40);
+	}
+
+	virtual TSoftClassPtr<UObject> GetAssetClass() const override
+	{
+		return UChaosCacheCollection::StaticClass();
+	}
+
+	virtual FText GetAssetDescription(const FAssetData& AssetData) const override
+	{
+		return NSLOCTEXT("AssetDefinition", "AssetDefinition_ChaosCacheCollectionDescription", "A collection of physically active component caches that can be used to record and replay Chaos simulation recordings.");
+	}
+
+	virtual TConstArrayView<FAssetCategoryPath> GetAssetCategories() const override
+	{
+		static const TArray<FAssetCategoryPath> Categories = 
+			{
+				FAssetCategoryPath(EAssetCategoryPaths::Physics, NSLOCTEXT("AssetDefinition", "AssetDefinition_ChaosCacheCollectionSubMenu", "Chaos"), ECategoryMenuType::Section)
+			};
+		return Categories;
+	}
+	// UAssetDefinition End
+};
diff --git a/Engine/Plugins/Experimental/ChaosCaching/Source/ChaosCachingEditor/Private/Chaos/AssetTypeActions_ChaosCacheCollection.cpp b/Engine/Plugins/Experimental/ChaosCaching/Source/ChaosCachingEditor/Private/Chaos/AssetTypeActions_ChaosCacheCollection.cpp
deleted file mode 100644
diff --git a/Engine/Plugins/Experimental/ChaosCaching/Source/ChaosCachingEditor/Private/Chaos/ChaosCachingEditorPlugin.cpp b/Engine/Plugins/Experimental/ChaosCaching/Source/ChaosCachingEditor/Private/Chaos/ChaosCachingEditorPlugin.cpp
--- a/Engine/Plugins/Experimental/ChaosCaching/Source/ChaosCachingEditor/Private/Chaos/ChaosCachingEditorPlugin.cpp
+++ b/Engine/Plugins/Experimental/ChaosCaching/Source/ChaosCachingEditor/Private/Chaos/ChaosCachingEditorPlugin.cpp
@@ -3,7 +3,6 @@
 #include "Chaos/ChaosCachingEditorPlugin.h"
 #include "Chaos/ChaosCachingEditorStyle.h"
 #include "Chaos/Adapters/CacheAdapter.h"
-#include "Chaos/AssetTypeActions_ChaosCacheCollection.h"
 #include "Chaos/CacheCollectionCustomization.h"
 #include "Chaos/CacheEditorCommands.h"
 #include "Chaos/CacheManagerActor.h"
@@ -23,6 +22,7 @@
 #include "Framework/Application/SlateApplication.h"
 #include "Framework/MultiBox/MultiBoxBuilder.h"
 #include "ActorTreeItem.h"
+#include "AssetToolsModule.h"
 #include "SceneOutlinerModule.h"
 #include "Selection.h"
 #include "TakeRecorderSources.h"
@@ -173,12 +173,6 @@ static const FName MovieSceneTrackRecorderFactoryName("MovieSceneTrackRecorderFa
 
 void IChaosCachingEditorPlugin::StartupModule()
 {
-	AssetTypeActions_ChaosCacheCollection = new FAssetTypeActions_ChaosCacheCollection();
-
-	FAssetToolsModule& AssetToolsModule = FAssetToolsModule::GetModule();
-	IAssetTools&       AssetTools       = AssetToolsModule.Get();
-	AssetTools.RegisterAssetTypeActions(MakeShareable(AssetTypeActions_ChaosCacheCollection));
-
 	FCachingEditorCommands::Register();
 
 	// Register level editor menu extender
@@ -222,11 +216,6 @@ void IChaosCachingEditorPlugin::ShutdownModule()
 
 		FCachingEditorCommands::Unregister();
 
-		FAssetToolsModule& AssetToolsModule = FAssetToolsModule::GetModule();
-		IAssetTools&       AssetTools       = AssetToolsModule.Get();
-
-		AssetTools.UnregisterAssetTypeActions(AssetTypeActions_ChaosCacheCollection->AsShared());
-
 		ISequencerModule* SequencerModulePtr = FModuleManager::Get().GetModulePtr<ISequencerModule>("Sequencer");
 		if (SequencerModulePtr)
 		{
diff --git a/Engine/Plugins/Experimental/ChaosCaching/Source/ChaosCachingEditor/Public/Chaos/AssetTypeActions_ChaosCacheCollection.h b/Engine/Plugins/Experimental/ChaosCaching/Source/ChaosCachingEditor/Public/Chaos/AssetTypeActions_ChaosCacheCollection.h
deleted file mode 100644
diff --git a/Engine/Plugins/Experimental/ChaosCaching/Source/ChaosCachingEditor/Public/Chaos/ChaosCachingEditorPlugin.h b/Engine/Plugins/Experimental/ChaosCaching/Source/ChaosCachingEditor/Public/Chaos/ChaosCachingEditorPlugin.h
--- a/Engine/Plugins/Experimental/ChaosCaching/Source/ChaosCachingEditor/Public/Chaos/ChaosCachingEditorPlugin.h
+++ b/Engine/Plugins/Experimental/ChaosCaching/Source/ChaosCachingEditor/Public/Chaos/ChaosCachingEditorPlugin.h
@@ -52,7 +52,6 @@ private:
 	
 	void OnCreateCacheManager();
 
-	FAssetTypeActions_ChaosCacheCollection* AssetTypeActions_ChaosCacheCollection;
 	FDelegateHandle StartupHandle;
 	FDelegateHandle TrackEditorBindingHandle;
 	FDelegateHandle SourcesMenuExtension;
diff --git a/Engine/Plugins/Experimental/ChaosOutfitAsset/Source/ChaosOutfitAssetEditor/Private/ChaosOutfitAsset/AssetDefinition_OutfitAsset.cpp b/Engine/Plugins/Experimental/ChaosOutfitAsset/Source/ChaosOutfitAssetEditor/Private/ChaosOutfitAsset/AssetDefinition_OutfitAsset.cpp
--- a/Engine/Plugins/Experimental/ChaosOutfitAsset/Source/ChaosOutfitAssetEditor/Private/ChaosOutfitAsset/AssetDefinition_OutfitAsset.cpp
+++ b/Engine/Plugins/Experimental/ChaosOutfitAsset/Source/ChaosOutfitAssetEditor/Private/ChaosOutfitAsset/AssetDefinition_OutfitAsset.cpp
@@ -32,7 +32,10 @@ FLinearColor UAssetDefinition_OutfitAsset::GetAssetColor() const
 
 TConstArrayView<FAssetCategoryPath> UAssetDefinition_OutfitAsset::GetAssetCategories() const
 {
-	static const auto Categories = { EAssetCategoryPaths::Physics };
+	static const TArray<FAssetCategoryPath> Categories = 
+		{
+			FAssetCategoryPath(EAssetCategoryPaths::Physics, LOCTEXT("AssetDefinition_OutfitAssetSubMenu", "Chaos Cloth"), ECategoryMenuType::Section)
+		};
 	return Categories;
 }
 
diff --git a/Engine/Plugins/Experimental/ChaosSolverPlugin/Source/ChaosSolverEditor/ChaosSolverEditor.Build.cs b/Engine/Plugins/Experimental/ChaosSolverPlugin/Source/ChaosSolverEditor/ChaosSolverEditor.Build.cs
--- a/Engine/Plugins/Experimental/ChaosSolverPlugin/Source/ChaosSolverEditor/ChaosSolverEditor.Build.cs
+++ b/Engine/Plugins/Experimental/ChaosSolverPlugin/Source/ChaosSolverEditor/ChaosSolverEditor.Build.cs
@@ -30,6 +30,7 @@ namespace UnrealBuildTool.Rules
 			PrivateDependencyModuleNames.AddRange(
 				new string[]
 				{
+					"AssetDefinition",
 					"EditorStyle",
 				}
 				);
diff --git a/Engine/Plugins/Experimental/ChaosSolverPlugin/Source/ChaosSolverEditor/Private/Chaos/AssetDefinition_ChaosSolver.h b/Engine/Plugins/Experimental/ChaosSolverPlugin/Source/ChaosSolverEditor/Private/Chaos/AssetDefinition_ChaosSolver.h
new file mode 100644
--- /dev/null
+++ b/Engine/Plugins/Experimental/ChaosSolverPlugin/Source/ChaosSolverEditor/Private/Chaos/AssetDefinition_ChaosSolver.h
@@ -0,0 +1,40 @@
+// Copyright Epic Games, Inc. All Rights Reserved.
+
+#pragma once
+#include "AssetDefinitionDefault.h"
+#include "Chaos/ChaosSolver.h"
+#include "AssetDefinition_ChaosSolver.generated.h"
+
+UCLASS()
+class UAssetDefinition_ChaosSolver : public UAssetDefinitionDefault
+{
+	GENERATED_BODY()
+public:
+	// UAssetDefinition Begin
+	virtual FText GetAssetDisplayName() const override
+	{
+		return NSLOCTEXT("AssetDefinition", "AssetDefinition_ChaosSolver", "Chaos Solver");
+	}
+
+	virtual FLinearColor GetAssetColor() const override
+	{
+		return FColor(255, 192, 128);
+	}
+
+	virtual TSoftClassPtr<UObject> GetAssetClass() const override
+	{
+		return UChaosSolver::StaticClass();
+	}
+
+	virtual TConstArrayView<FAssetCategoryPath> GetAssetCategories() const override
+	{
+		static const TArray<FAssetCategoryPath> Categories = 
+			{
+				FAssetCategoryPath(EAssetCategoryPaths::Physics, NSLOCTEXT("AssetDefinition", "AssetDefinition_ChaosSolverSubMenu", "Chaos"), ECategoryMenuType::Section)
+			};
+		return Categories;
+	}
+
+	virtual EAssetCommandResult OpenAssets(const FAssetOpenArgs& OpenArgs) const override { return EAssetCommandResult::Handled; }
+	// UAssetDefinition End
+};
diff --git a/Engine/Plugins/Experimental/ChaosSolverPlugin/Source/ChaosSolverEditor/Private/Chaos/AssetTypeActions_ChaosSolver.cpp b/Engine/Plugins/Experimental/ChaosSolverPlugin/Source/ChaosSolverEditor/Private/Chaos/AssetTypeActions_ChaosSolver.cpp
deleted file mode 100644
diff --git a/Engine/Plugins/Experimental/ChaosSolverPlugin/Source/ChaosSolverEditor/Private/Chaos/ChaosSolverEditorPlugin.cpp b/Engine/Plugins/Experimental/ChaosSolverPlugin/Source/ChaosSolverEditor/Private/Chaos/ChaosSolverEditorPlugin.cpp
--- a/Engine/Plugins/Experimental/ChaosSolverPlugin/Source/ChaosSolverEditor/Private/Chaos/ChaosSolverEditorPlugin.cpp
+++ b/Engine/Plugins/Experimental/ChaosSolverPlugin/Source/ChaosSolverEditor/Private/Chaos/ChaosSolverEditorPlugin.cpp
@@ -3,7 +3,6 @@
 
 #include "Chaos/ChaosSolverEditorPlugin.h"
 
-#include "Chaos/AssetTypeActions_ChaosSolver.h"
 #include "Chaos/ChaosSolverEditorStyle.h"
 #include "PropertyEditorModule.h"
 #include "ChaosSolverEditorDetails.h"
@@ -16,11 +15,6 @@ void IChaosSolverEditorPlugin::StartupModule()
 {
 	FChaosSolverEditorStyle::Get();
 
-	FAssetToolsModule& AssetToolsModule = FAssetToolsModule::GetModule();
-	IAssetTools& AssetTools = AssetToolsModule.Get();
-	AssetTypeActions_ChaosSolver = new FAssetTypeActions_ChaosSolver();
-	AssetTools.RegisterAssetTypeActions(MakeShareable(AssetTypeActions_ChaosSolver));
-
 	if (GIsEditor && !IsRunningCommandlet())
 	{
 	}
@@ -33,13 +27,6 @@ void IChaosSolverEditorPlugin::StartupModule()
 
 void IChaosSolverEditorPlugin::ShutdownModule()
 {
-	if (UObjectInitialized())
-	{
-		FAssetToolsModule& AssetToolsModule = FAssetToolsModule::GetModule();
-		IAssetTools& AssetTools = AssetToolsModule.Get();
-		AssetTools.UnregisterAssetTypeActions(AssetTypeActions_ChaosSolver->AsShared());
-	}
-
 	// Unregister details view customizations
 	FPropertyEditorModule& PropertyModule = FModuleManager::GetModuleChecked<FPropertyEditorModule>("PropertyEditor");
 	PropertyModule.UnregisterCustomPropertyTypeLayout("ChaosDebugSubstepControl");
diff --git a/Engine/Plugins/Experimental/ChaosSolverPlugin/Source/ChaosSolverEditor/Public/Chaos/AssetTypeActions_ChaosSolver.h b/Engine/Plugins/Experimental/ChaosSolverPlugin/Source/ChaosSolverEditor/Public/Chaos/AssetTypeActions_ChaosSolver.h
deleted file mode 100644
diff --git a/Engine/Plugins/Experimental/ChaosSolverPlugin/Source/ChaosSolverEditor/Public/Chaos/ChaosSolverEditorPlugin.h b/Engine/Plugins/Experimental/ChaosSolverPlugin/Source/ChaosSolverEditor/Public/Chaos/ChaosSolverEditorPlugin.h
--- a/Engine/Plugins/Experimental/ChaosSolverPlugin/Source/ChaosSolverEditor/Public/Chaos/ChaosSolverEditorPlugin.h
+++ b/Engine/Plugins/Experimental/ChaosSolverPlugin/Source/ChaosSolverEditor/Public/Chaos/ChaosSolverEditorPlugin.h
@@ -39,8 +39,4 @@ public:
 	{
 		return FModuleManager::Get().IsModuleLoaded( "ChaosSolverEditorPlugin" );
 	}
-
-private:
-	class FAssetTypeActions_ChaosSolver* AssetTypeActions_ChaosSolver;
-
 };
diff --git a/Engine/Plugins/Experimental/Dataflow/Source/DataflowEditor/Private/Dataflow/AssetDefinition_DataflowAsset.cpp b/Engine/Plugins/Experimental/Dataflow/Source/DataflowEditor/Private/Dataflow/AssetDefinition_DataflowAsset.cpp
--- a/Engine/Plugins/Experimental/Dataflow/Source/DataflowEditor/Private/Dataflow/AssetDefinition_DataflowAsset.cpp
+++ b/Engine/Plugins/Experimental/Dataflow/Source/DataflowEditor/Private/Dataflow/AssetDefinition_DataflowAsset.cpp
@@ -224,7 +224,7 @@ namespace UE::DataflowAssetDefinitionHelpers
 	 *   - Existing asset : sets the selected shared UDataflow on Asset without duplication.
 	 *   - None / Cancelled : leaves Asset unchanged.
 	 *
-	 * Always returns true — the picker is a single non-repeating step with no sub-dialogs.
+	 * Always returns true - the picker is a single non-repeating step with no sub-dialogs.
 	 */
 	bool SetDataflowFromTemplatePicker(UObject* Asset, bool bShowDefaultBlankOption)
 	{
@@ -251,7 +251,7 @@ namespace UE::DataflowAssetDefinitionHelpers
 				: LoadObject<UDataflow>(nullptr, *TemplatePath);
 
 			// Embed a duplicate of the template, or a new empty graph for Blank.
-			// Embedded dataflows are owned by the asset — no standalone/public flags.
+			// Embedded dataflows are owned by the asset - no standalone/public flags.
 			UDataflow* Embedded = Source
 				? DuplicateObject<UDataflow>(Source, Asset, TEXT("EmbeddedDataflow"))
 				: NewObject<UDataflow>(Asset, UDataflow::StaticClass(), TEXT("EmbeddedDataflow"), RF_Transactional);
@@ -271,7 +271,7 @@ namespace UE::DataflowAssetDefinitionHelpers
 
 		case EDataflowPickerResult::ExistingAsset:
 		{
-			// Reference the selected asset directly — no duplication.
+			// Reference the selected asset directly - no duplication.
 			if (UDataflow* const SharedDataflow = Cast<UDataflow>(PickerResult.SelectedAsset.GetAsset()))
 			{
 				UE::Dataflow::InstanceUtils::SetDataflowAssetOnObject(SharedDataflow, Asset);
@@ -683,7 +683,10 @@ FLinearColor UAssetDefinition_DataflowAsset::GetAssetColor() const
 
 TConstArrayView<FAssetCategoryPath> UAssetDefinition_DataflowAsset::GetAssetCategories() const
 {
-	static const auto Categories = { EAssetCategoryPaths::Physics };
+	static const TArray<FAssetCategoryPath> Categories = 
+		{
+			FAssetCategoryPath(EAssetCategoryPaths::Physics, LOCTEXT("DataflowAsset_SubMenu", "Dataflow"), ECategoryMenuType::Section)
+		};
 	return Categories;
 }
 
diff --git a/Engine/Plugins/Experimental/GeometryCollectionPlugin/Source/GeometryCollectionEditor/Private/GeometryCollection/AssetDefinition_GeometryCollection.cpp b/Engine/Plugins/Experimental/GeometryCollectionPlugin/Source/GeometryCollectionEditor/Private/GeometryCollection/AssetDefinition_GeometryCollection.cpp
--- a/Engine/Plugins/Experimental/GeometryCollectionPlugin/Source/GeometryCollectionEditor/Private/GeometryCollection/AssetDefinition_GeometryCollection.cpp
+++ b/Engine/Plugins/Experimental/GeometryCollectionPlugin/Source/GeometryCollectionEditor/Private/GeometryCollection/AssetDefinition_GeometryCollection.cpp
@@ -54,7 +54,10 @@ FLinearColor UAssetDefinition_GeometryCollection::GetAssetColor() const
 
 TConstArrayView<FAssetCategoryPath> UAssetDefinition_GeometryCollection::GetAssetCategories() const
 {
-	static const auto Categories = { EAssetCategoryPaths::Physics };
+	static const TArray<FAssetCategoryPath> Categories = 
+		{
+			FAssetCategoryPath(EAssetCategoryPaths::Physics, LOCTEXT("GeometryCollection_SubMenu", "Geometry Collection"), ECategoryMenuType::Section)
+		};
 	return Categories;
 }
 
diff --git a/Engine/Plugins/Experimental/GeometryCollectionPlugin/Source/GeometryCollectionEditor/Private/GeometryCollection/AssetTypeActions_GeometryCollectionCache.cpp b/Engine/Plugins/Experimental/GeometryCollectionPlugin/Source/GeometryCollectionEditor/Private/GeometryCollection/AssetTypeActions_GeometryCollectionCache.cpp
deleted file mode 100644
diff --git a/Engine/Plugins/Experimental/GeometryCollectionPlugin/Source/GeometryCollectionEditor/Private/GeometryCollection/GeometryCollectionEditorPlugin.cpp b/Engine/Plugins/Experimental/GeometryCollectionPlugin/Source/GeometryCollectionEditor/Private/GeometryCollection/GeometryCollectionEditorPlugin.cpp
--- a/Engine/Plugins/Experimental/GeometryCollectionPlugin/Source/GeometryCollectionEditor/Private/GeometryCollection/GeometryCollectionEditorPlugin.cpp
+++ b/Engine/Plugins/Experimental/GeometryCollectionPlugin/Source/GeometryCollectionEditor/Private/GeometryCollection/GeometryCollectionEditorPlugin.cpp
@@ -14,7 +14,6 @@
 #include "GeometryCollection/GeometryCollectionConversion.h"
 #include "GeometryCollection/GeometryCollectionCommands.h"
 #include "GeometryCollection/AssetDefinition_GeometryCollection.h"
-#include "GeometryCollection/AssetTypeActions_GeometryCollectionCache.h"
 #include "GeometryCollection/GeometryCollectionAssetBroker.h"
 #include "GeometryCollection/GeometryCollectionEditorStyle.h"
 #include "GeometryCollection/GeometryCollectionThumbnailRenderer.h"
@@ -46,11 +45,6 @@ void IGeometryCollectionEditorPlugin::StartupModule()
 
 	FGeometryCollectionSelectionCommands::Register();
 
-	FAssetToolsModule& AssetToolsModule = FAssetToolsModule::GetModule();
-	IAssetTools& AssetTools = AssetToolsModule.Get();
-	GeometryCollectionCacheAssetActions = new FAssetTypeActions_GeometryCollectionCache();
-	AssetTools.RegisterAssetTypeActions(MakeShareable(GeometryCollectionCacheAssetActions));
-
 	AssetBroker = new FGeometryCollectionAssetBroker();
 	FComponentAssetBrokerage::RegisterBroker(MakeShareable(AssetBroker), UGeometryCollectionComponent::StaticClass(), true, true);
 
@@ -278,10 +272,6 @@ void IGeometryCollectionEditorPlugin::ShutdownModule()
 
 		UThumbnailManager::Get().UnregisterCustomRenderer(UGeometryCollection::StaticClass());
 
-		FAssetToolsModule& AssetToolsModule = FAssetToolsModule::GetModule();
-		IAssetTools& AssetTools = AssetToolsModule.Get();
-		AssetTools.UnregisterAssetTypeActions(GeometryCollectionCacheAssetActions->AsShared());
-
 		FComponentAssetBrokerage::UnregisterBroker(MakeShareable(AssetBroker));
 
 		// Unbind provider from editor
diff --git a/Engine/Plugins/Experimental/GeometryCollectionPlugin/Source/GeometryCollectionEditor/Public/GeometryCollection/AssetTypeActions_GeometryCollectionCache.h b/Engine/Plugins/Experimental/GeometryCollectionPlugin/Source/GeometryCollectionEditor/Public/GeometryCollection/AssetTypeActions_GeometryCollectionCache.h
deleted file mode 100644
diff --git a/Engine/Plugins/Experimental/GeometryCollectionPlugin/Source/GeometryCollectionEditor/Public/GeometryCollection/GeometryCollectionEditorPlugin.h b/Engine/Plugins/Experimental/GeometryCollectionPlugin/Source/GeometryCollectionEditor/Public/GeometryCollection/GeometryCollectionEditorPlugin.h
--- a/Engine/Plugins/Experimental/GeometryCollectionPlugin/Source/GeometryCollectionEditor/Public/GeometryCollection/GeometryCollectionEditorPlugin.h
+++ b/Engine/Plugins/Experimental/GeometryCollectionPlugin/Source/GeometryCollectionEditor/Public/GeometryCollection/GeometryCollectionEditorPlugin.h
@@ -64,7 +64,6 @@ private:
 	TArray<IConsoleObject*> EditorCommands;
 
 	// Asset actions for new asset types
-	FAssetTypeActions_GeometryCollectionCache* GeometryCollectionCacheAssetActions;
 	FGeometryCollectionAssetBroker* AssetBroker;
 
 	// Modular features
diff --git a/Engine/Plugins/Experimental/PhysicsControl/Source/PhysicsControlEditor/PhysicsControlEditor.Build.cs b/Engine/Plugins/Experimental/PhysicsControl/Source/PhysicsControlEditor/PhysicsControlEditor.Build.cs
--- a/Engine/Plugins/Experimental/PhysicsControl/Source/PhysicsControlEditor/PhysicsControlEditor.Build.cs
+++ b/Engine/Plugins/Experimental/PhysicsControl/Source/PhysicsControlEditor/PhysicsControlEditor.Build.cs
@@ -41,6 +41,7 @@ public class PhysicsControlEditor : ModuleRules
 				"AnimGraph",
 				"AnimGraphRuntime",
 				"ApplicationCore",
+				"AssetDefinition",
 				"BlueprintGraph",
 				"Chaos",
 				"ComponentVisualizers",
diff --git a/Engine/Plugins/Experimental/PhysicsControl/Source/PhysicsControlEditor/Private/AssetDefinition_PhysicsControlAsset.cpp b/Engine/Plugins/Experimental/PhysicsControl/Source/PhysicsControlEditor/Private/AssetDefinition_PhysicsControlAsset.cpp
new file mode 100644
--- /dev/null
+++ b/Engine/Plugins/Experimental/PhysicsControl/Source/PhysicsControlEditor/Private/AssetDefinition_PhysicsControlAsset.cpp
@@ -0,0 +1,22 @@
+// Copyright Epic Games, Inc. All Rights Reserved.
+
+#include "AssetDefinition_PhysicsControlAsset.h"
+
+#include "PhysicsControlAssetEditor.h"
+
+#define LOCTEXT_NAMESPACE "AssetDefinition_PhysicsControlAsset"
+
+EAssetCommandResult UAssetDefinition_PhysicsControlAsset::OpenAssets(const FAssetOpenArgs& OpenArgs) const
+{
+	const EToolkitMode::Type Mode = 
+		OpenArgs.ToolkitHost.IsValid() ? EToolkitMode::WorldCentric : EToolkitMode::Standalone;
+
+	for (UPhysicsControlAsset* PhysicsControlAsset : OpenArgs.LoadObjects<UPhysicsControlAsset>())
+	{
+		const TSharedRef<FPhysicsControlAssetEditor> NewEditor(new FPhysicsControlAssetEditor());
+		NewEditor->InitAssetEditor(Mode, OpenArgs.ToolkitHost, PhysicsControlAsset);
+	}
+	return EAssetCommandResult::Handled;
+}
+
+#undef LOCTEXT_NAMESPACE
diff --git a/Engine/Plugins/Experimental/PhysicsControl/Source/PhysicsControlEditor/Private/AssetDefinition_PhysicsControlAsset.h b/Engine/Plugins/Experimental/PhysicsControl/Source/PhysicsControlEditor/Private/AssetDefinition_PhysicsControlAsset.h
new file mode 100644
--- /dev/null
+++ b/Engine/Plugins/Experimental/PhysicsControl/Source/PhysicsControlEditor/Private/AssetDefinition_PhysicsControlAsset.h
@@ -0,0 +1,40 @@
+// Copyright Epic Games, Inc. All Rights Reserved.
+
+#pragma once
+#include "AssetDefinitionDefault.h"
+#include "PhysicsControlAsset.h"
+#include "AssetDefinition_PhysicsControlAsset.generated.h"
+
+UCLASS()
+class UAssetDefinition_PhysicsControlAsset : public UAssetDefinitionDefault
+{
+	GENERATED_BODY()
+public:
+	// UAssetDefinition Begin
+	virtual FText GetAssetDisplayName() const override
+	{
+		return NSLOCTEXT("AssetDefinition", "AssetDefinition_PhysicsControlAsset", "Physics Control Asset");
+	}
+
+	virtual FLinearColor GetAssetColor() const override
+	{
+		return FColor(255, 192, 128);
+	}
+
+	virtual TSoftClassPtr<UObject> GetAssetClass() const override
+	{
+		return UPhysicsControlAsset::StaticClass();
+	}
+
+	virtual TConstArrayView<FAssetCategoryPath> GetAssetCategories() const override
+	{
+		static const TArray<FAssetCategoryPath> Categories = 
+		{
+			FAssetCategoryPath(EAssetCategoryPaths::Physics, NSLOCTEXT("AssetDefinition", "AssetDefinition_PhysicsControlAssetSubMenu", "Physical Material"), ECategoryMenuType::Section)
+		};
+		return Categories;
+	}
+
+	virtual EAssetCommandResult OpenAssets(const FAssetOpenArgs& OpenArgs) const override;
+	// UAssetDefinition End
+};
diff --git a/Engine/Plugins/Experimental/PhysicsControl/Source/PhysicsControlEditor/Private/PhysicsControlAssetActions.cpp b/Engine/Plugins/Experimental/PhysicsControl/Source/PhysicsControlEditor/Private/PhysicsControlAssetActions.cpp
deleted file mode 100644
diff --git a/Engine/Plugins/Experimental/PhysicsControl/Source/PhysicsControlEditor/Private/PhysicsControlAssetActions.h b/Engine/Plugins/Experimental/PhysicsControl/Source/PhysicsControlEditor/Private/PhysicsControlAssetActions.h
deleted file mode 100644
diff --git a/Engine/Plugins/Experimental/PhysicsControl/Source/PhysicsControlEditor/Private/PhysicsControlAssetInfoDetailsCustomization.cpp b/Engine/Plugins/Experimental/PhysicsControl/Source/PhysicsControlEditor/Private/PhysicsControlAssetInfoDetailsCustomization.cpp
--- a/Engine/Plugins/Experimental/PhysicsControl/Source/PhysicsControlEditor/Private/PhysicsControlAssetInfoDetailsCustomization.cpp
+++ b/Engine/Plugins/Experimental/PhysicsControl/Source/PhysicsControlEditor/Private/PhysicsControlAssetInfoDetailsCustomization.cpp
@@ -2,7 +2,6 @@
 
 #include "PhysicsControlAssetInfoDetailsCustomization.h"
 
-#include "PhysicsControlAssetActions.h"
 #include "PhysicsControlAssetEditor.h"
 #include "PhysicsControlAsset.h"
 #include "PhysicsControlAssetEditorData.h"
diff --git a/Engine/Plugins/Experimental/PhysicsControl/Source/PhysicsControlEditor/Private/PhysicsControlAssetPreviewDetailsCustomization.cpp b/Engine/Plugins/Experimental/PhysicsControl/Source/PhysicsControlEditor/Private/PhysicsControlAssetPreviewDetailsCustomization.cpp
--- a/Engine/Plugins/Experimental/PhysicsControl/Source/PhysicsControlEditor/Private/PhysicsControlAssetPreviewDetailsCustomization.cpp
+++ b/Engine/Plugins/Experimental/PhysicsControl/Source/PhysicsControlEditor/Private/PhysicsControlAssetPreviewDetailsCustomization.cpp
@@ -2,7 +2,6 @@
 
 #include "PhysicsControlAssetPreviewDetailsCustomization.h"
 
-#include "PhysicsControlAssetActions.h"
 #include "PhysicsControlAssetEditor.h"
 #include "PhysicsControlAssetEditorData.h"
 #include "PhysicsControlAssetEditorCommands.h"
diff --git a/Engine/Plugins/Experimental/PhysicsControl/Source/PhysicsControlEditor/Private/PhysicsControlAssetSetupDetailsCustomization.cpp b/Engine/Plugins/Experimental/PhysicsControl/Source/PhysicsControlEditor/Private/PhysicsControlAssetSetupDetailsCustomization.cpp
--- a/Engine/Plugins/Experimental/PhysicsControl/Source/PhysicsControlEditor/Private/PhysicsControlAssetSetupDetailsCustomization.cpp
+++ b/Engine/Plugins/Experimental/PhysicsControl/Source/PhysicsControlEditor/Private/PhysicsControlAssetSetupDetailsCustomization.cpp
@@ -2,7 +2,6 @@
 
 #include "PhysicsControlAssetSetupDetailsCustomization.h"
 
-#include "PhysicsControlAssetActions.h"
 #include "PhysicsControlAssetEditor.h"
 #include "PhysicsControlAsset.h"
 #include "PhysicsControlAssetEditorData.h"
diff --git a/Engine/Plugins/Experimental/PhysicsControl/Source/PhysicsControlEditor/Private/PhysicsControlEditorModule.cpp b/Engine/Plugins/Experimental/PhysicsControl/Source/PhysicsControlEditor/Private/PhysicsControlEditorModule.cpp
--- a/Engine/Plugins/Experimental/PhysicsControl/Source/PhysicsControlEditor/Private/PhysicsControlEditorModule.cpp
+++ b/Engine/Plugins/Experimental/PhysicsControl/Source/PhysicsControlEditor/Private/PhysicsControlEditorModule.cpp
@@ -4,21 +4,17 @@
 #include "PhysicsControlComponent.h"
 #include "PhysicsControlComponentVisualizer.h"
 #include "OperatorViewer/OperatorViewer.h"
-#include "PhysicsControlAssetActions.h"
 #include "PhysicsControlAssetEditorEditMode.h"
 #include "UnrealEdGlobals.h"
 #include "Editor/UnrealEdEngine.h"
 #include "Features/IModularFeatures.h"
+#include "Modules/ModuleManager.h"
 
 #define LOCTEXT_NAMESPACE "PhysicsControlModule"
 
 //======================================================================================================================
 void FPhysicsControlEditorModule::StartupModule()
 {
-	IAssetTools& AssetTools = FModuleManager::LoadModuleChecked<FAssetToolsModule>("AssetTools").Get();
-	PhysicsControlAssetActions = MakeShared<FPhysicsControlAssetActions>();
-	AssetTools.RegisterAssetTypeActions(PhysicsControlAssetActions.ToSharedRef());
-	
 	FEditorModeRegistry::Get().RegisterMode<FPhysicsControlAssetEditorEditMode>(
 		FPhysicsControlAssetEditorEditMode::ModeName, 
 		LOCTEXT("PhysicsControlAssetEditorMode", "PhysicsControlAsset"), 
@@ -59,12 +55,6 @@ void FPhysicsControlEditorModule::ShutdownModule()
 	
 	FEditorModeRegistry::Get().UnregisterMode(FPhysicsControlAssetEditorEditMode::ModeName);
 
-	if (FModuleManager::Get().IsModuleLoaded("AssetTools"))
-	{
-		FAssetToolsModule::GetModule().Get().UnregisterAssetTypeActions(
-			PhysicsControlAssetActions.ToSharedRef());
-	}
-
 	if (GUnrealEd)
 	{
 		for (const FName& Name : VisualizersToUnregisterOnShutdown)
diff --git a/Engine/Plugins/Experimental/PhysicsControl/Source/PhysicsControlEditor/Public/PhysicsControlEditorModule.h b/Engine/Plugins/Experimental/PhysicsControl/Source/PhysicsControlEditor/Public/PhysicsControlEditorModule.h
--- a/Engine/Plugins/Experimental/PhysicsControl/Source/PhysicsControlEditor/Public/PhysicsControlEditorModule.h
+++ b/Engine/Plugins/Experimental/PhysicsControl/Source/PhysicsControlEditor/Public/PhysicsControlEditorModule.h
@@ -24,8 +24,6 @@ public:
 private:
 
 	TArray<FName> VisualizersToUnregisterOnShutdown;
-	TSharedPtr<FPhysicsControlAssetActions> PhysicsControlAssetActions;
-
 	FPhysicsControlOperatorViewer* EditorInterface = nullptr;
 };
 
diff --git a/Engine/Plugins/Runtime/HairStrands/Source/HairStrandsEditor/Public/GroomActions.h b/Engine/Plugins/Runtime/HairStrands/Source/HairStrandsEditor/Public/GroomActions.h
--- a/Engine/Plugins/Runtime/HairStrands/Source/HairStrandsEditor/Public/GroomActions.h
+++ b/Engine/Plugins/Runtime/HairStrands/Source/HairStrandsEditor/Public/GroomActions.h
@@ -15,15 +15,32 @@ class UAssetDefinition_GroomAsset : public UAssetDefinitionDefault
 
 public:
 	// UAssetDefinition Begin
-	virtual FText GetAssetDisplayName() const override { return NSLOCTEXT("AssetTypeActions", "AssetTypeActions_Groom", "Groom Asset"); }
+	virtual FText GetAssetDisplayName() const override
+	{
+		return NSLOCTEXT("AssetDefinition", "AssetDefinition_Groom", "Groom Asset");
+	}
+
 	virtual FLinearColor GetAssetColor() const override;
-	virtual TSoftClassPtr<UObject> GetAssetClass() const override { return UGroomAsset::StaticClass(); }
-	virtual bool CanImport() const override { return true; }
+	
+	virtual TSoftClassPtr<UObject> GetAssetClass() const override
+	{
+		return UGroomAsset::StaticClass();
+	}
+
+	virtual bool CanImport() const override
+	{
+		return true;
+	}
+
 	virtual TConstArrayView<FAssetCategoryPath> GetAssetCategories() const override
 	{
-		static const auto Categories = { EAssetCategoryPaths::Physics };
+		static const TArray<FAssetCategoryPath> Categories = 
+		{
+			FAssetCategoryPath(EAssetCategoryPaths::Physics, NSLOCTEXT("AssetDefinition", "AssetDefinition_GroomSubMenu", "Other"), ECategoryMenuType::Section)
+		};
 		return Categories;
 	}
+
 	virtual UThumbnailInfo* LoadThumbnailInfo(const FAssetData& InAssetData) const override;
 
 	virtual EAssetCommandResult OpenAssets(const FAssetOpenArgs& OpenArgs) const override;
diff --git a/Engine/Source/Developer/AssetTools/Private/AssetTypeActions/AssetTypeActions_PhysicalMaterialMask.cpp b/Engine/Source/Developer/AssetTools/Private/AssetTypeActions/AssetTypeActions_PhysicalMaterialMask.cpp
deleted file mode 100644
diff --git a/Engine/Source/Developer/AssetTools/Private/AssetTypeActions/AssetTypeActions_PhysicalMaterialMask.h b/Engine/Source/Developer/AssetTools/Private/AssetTypeActions/AssetTypeActions_PhysicalMaterialMask.h
deleted file mode 100644
